ConcertClient early out if there is no bound delegate.

diff --git a/Engine/Plugins/Developer/Concert/ConcertSync/ConcertSyncClient/Source/ConcertSyncClient/Private/ConcertClientPackageBridge.cpp b/Engine/Plugins/Developer/Concert/ConcertSync/ConcertSyncClient/Source/ConcertSyncClient/Private/ConcertClientPackageBridge.cpp
--- a/Engine/Plugins/Developer/Concert/ConcertSync/ConcertSyncClient/Source/ConcertSyncClient/Private/ConcertClientPackageBridge.cpp
+++ b/Engine/Plugins/Developer/Concert/ConcertSync/ConcertSyncClient/Source/ConcertSyncClient/Private/ConcertClientPackageBridge.cpp
@@ -115,6 +115,12 @@ void FConcertClientPackageBridge::HandlePackagePreSave(UPackage* Package)
 		return;
 	}
 
+	// Early out if the delegate is unbound
+	if (!OnLocalPackageEventDelegate.IsBound())
+	{
+		return;
+	}
+
 	UWorld* World = UWorld::FindWorldInPackage(Package);
 
 	FString PackageFilename;
@@ -144,6 +150,12 @@ void FConcertClientPackageBridge::HandlePackageSaved(const FString& PackageFilen
 		return;
 	}
 
+	// Early out if the delegate is unbound
+	if (!OnLocalPackageEventDelegate.IsBound())
+	{
+		return;
+	}
+
 	// if we end up here, the package should be either unlocked or locked by this client, the server will resend the latest revision if it wasn't the case.
 	FName NewPackageName;
 	PackagesBeingRenamed.RemoveAndCopyValue(Package->GetFName(), NewPackageName);
@@ -164,6 +176,12 @@ void FConcertClientPackageBridge::HandlePackageSaved(const FString& PackageFilen
 
 void FConcertClientPackageBridge::HandleAssetAdded(UObject *Object)
 {
+	// Early out if the delegate is unbound
+	if (!OnLocalPackageEventDelegate.IsBound())
+	{
+		return;
+	}
+
 	UPackage* Package = Object->GetOutermost();
 
 	// Skip packages that are in the process of being renamed as they are always saved after being added
@@ -197,6 +215,12 @@ void FConcertClientPackageBridge::HandleAssetAdded(UObject *Object)
 
 void FConcertClientPackageBridge::HandleAssetDeleted(UObject *Object)
 {
+	// Early out if the delegate is unbound
+	if (!OnLocalPackageEventDelegate.IsBound())
+	{
+		return;
+	}
+
 	UPackage* Package = Object->GetOutermost();
 
 	FConcertPackage Event;
@@ -221,6 +245,12 @@ void FConcertClientPackageBridge::HandleAssetRenamed(const FAssetData& Data, con
 
 void FConcertClientPackageBridge::HandleAssetReload(const EPackageReloadPhase InPackageReloadPhase, FPackageReloadedEvent* InPackageReloadedEvent)
 {
+	// Early out if the delegate is unbound
+	if (!OnLocalPackageDiscardedDelegate.IsBound())
+	{
+		return;
+	}
+
 	if (InPackageReloadPhase == EPackageReloadPhase::PrePackageLoad)
 	{
 		UPackage* Package = const_cast<UPackage*>(InPackageReloadedEvent->GetOldPackage());
@@ -232,6 +262,12 @@ void FConcertClientPackageBridge::HandleAssetReload(const EPackageReloadPhase In
 
 void FConcertClientPackageBridge::HandleMapChanged(UWorld* InWorld, EMapChangeType InMapChangeType)
 {
+	// Early out if the delegate is unbound
+	if (!OnLocalPackageDiscardedDelegate.IsBound())
+	{
+		return;
+	}
+
 	if (InMapChangeType == EMapChangeType::TearDownWorld)
 	{
 		UPackage* Package = InWorld->GetOutermost();
