- Updated the icons and style.

diff --git a/Engine/Plugins/Developer/Concert/ConcertSync/ConcertSyncClient/Source/ConcertSyncUI/Private/ConcertFrontendStyle.cpp b/Engine/Plugins/Developer/Concert/ConcertSync/ConcertSyncClient/Source/ConcertSyncUI/Private/ConcertFrontendStyle.cpp
--- a/Engine/Plugins/Developer/Concert/ConcertSync/ConcertSyncClient/Source/ConcertSyncUI/Private/ConcertFrontendStyle.cpp
+++ b/Engine/Plugins/Developer/Concert/ConcertSync/ConcertSyncClient/Source/ConcertSyncUI/Private/ConcertFrontendStyle.cpp
@@ -38,7 +38,9 @@ void FConcertFrontendStyle::Initialize()
 	// Const icon sizes
 	const FVector2D Icon16x16(16.0f, 16.0f);
 	const FVector2D Icon20x20(20.0f, 20.0f);
+	const FVector2D Icon24x24(24.0f, 24.0f);
 	const FVector2D Icon40x40(40.0f, 40.0f);
+	const FVector2D Icon48x48(48.0f, 48.0f);
 
 	// 16x16
 	StyleSet->Set("Concert.Concert", new IMAGE_PLUGIN_BRUSH("Icons/icon_Concert_16x", Icon16x16));
@@ -54,6 +56,12 @@ void FConcertFrontendStyle::Initialize()
 	StyleSet->Set("Concert.Online.Small", new IMAGE_PLUGIN_BRUSH("Icons/icon_ConcertOnline_40x", Icon20x20));
 	StyleSet->Set("Concert.Offline.Small", new IMAGE_PLUGIN_BRUSH("Icons/icon_ConcertOffline_40x", Icon20x20));
 
+	// 24x24/48x48 -> For sequencer toolbar.
+	StyleSet->Set("Concert.Sequencer.SyncTimeline",       new IMAGE_PLUGIN_BRUSH("Icons/icon_SequencerSyncTimeline_48x", Icon48x48)); // Enable/disable playback and time scrubbing from a remote client.
+	StyleSet->Set("Concert.Sequencer.SyncTimeline.Small", new IMAGE_PLUGIN_BRUSH("Icons/icon_SequencerSyncTimeline_48x", Icon24x24));
+	StyleSet->Set("Concert.Sequencer.SyncSequence",       new IMAGE_PLUGIN_BRUSH("Icons/icon_SequencerSyncSequence_48x", Icon48x48)); // Allows or not a remote client to open/close sequencer.
+	StyleSet->Set("Concert.Sequencer.SyncSequence.Small", new IMAGE_PLUGIN_BRUSH("Icons/icon_SequencerSyncSequence_48x", Icon24x24));
+
 	// 40x40
 	StyleSet->Set("Concert.Online", new IMAGE_PLUGIN_BRUSH("Icons/icon_ConcertOnline_40x", Icon40x40));
 	StyleSet->Set("Concert.Offline", new IMAGE_PLUGIN_BRUSH("Icons/icon_ConcertOffline_40x", Icon40x40));
diff --git a/Engine/Plugins/Developer/Concert/MultiUserClient/Source/MultiUserClient/Private/ConcertWorkspaceUI.cpp b/Engine/Plugins/Developer/Concert/MultiUserClient/Source/MultiUserClient/Private/ConcertWorkspaceUI.cpp
--- a/Engine/Plugins/Developer/Concert/MultiUserClient/Source/MultiUserClient/Private/ConcertWorkspaceUI.cpp
+++ b/Engine/Plugins/Developer/Concert/MultiUserClient/Source/MultiUserClient/Private/ConcertWorkspaceUI.cpp
@@ -372,7 +372,7 @@ public:
 				NAME_None,
 				LOCTEXT("TogglePlaybackSyncLabel", "Playback Sync"),
 				LOCTEXT("TogglePlaybackSyncTooltip", "Toggle Multi-User Playback Sync. If the option is enabled, playback and scrubbing of Sequencer will be synchronized across all users in a Multi-Users session if they have that sequence open and this option enabled."),
-				FSlateIcon(FConcertFrontendStyle::GetStyleSetName(), "Concert.Concert", "Concert.Online.Small"), // TODO: Find appropriate icon
+				FSlateIcon(FConcertFrontendStyle::GetStyleSetName(), "Concert.Sequencer.SyncTimeline", "Concert.Sequencer.SyncTimeline.Small"),
 				EUserInterfaceActionType::ToggleButton
 			);
 
@@ -398,7 +398,7 @@ public:
 				NAME_None,
 				LOCTEXT("ToggleRemoteOpenLabel", "Remote Open"), 
 				LOCTEXT("ToggleRemoteOpenTooltip", "Toggle Multi-User Remote Open. If the option is enabled, opening a sequence will open the same sequence on all users in the Multi-User session that also have this option enabled."),
-				FSlateIcon(FConcertFrontendStyle::GetStyleSetName(), "Concert.Concert", "Concert.Offline.Small"), // TODO: Find appropriate icon
+				FSlateIcon(FConcertFrontendStyle::GetStyleSetName(), "Concert.Sequencer.SyncSequence", "Concert.Sequencer.SyncSequence.Small"),
 				EUserInterfaceActionType::ToggleButton
 			);
 		}
