- Added missing include to SImage.h

diff --git a/Engine/Plugins/Developer/Concert/ConcertSync/ConcertSyncClient/Source/ConcertSyncUI/Public/ConcertFrontendUtils.h b/Engine/Plugins/Developer/Concert/ConcertSync/ConcertSyncClient/Source/ConcertSyncUI/Public/ConcertFrontendUtils.h
--- a/Engine/Plugins/Developer/Concert/ConcertSync/ConcertSyncClient/Source/ConcertSyncUI/Public/ConcertFrontendUtils.h
+++ b/Engine/Plugins/Developer/Concert/ConcertSync/ConcertSyncClient/Source/ConcertSyncUI/Public/ConcertFrontendUtils.h
@@ -6,6 +6,7 @@
 #include "ConcertActionDefinition.h"
 #include "EditorStyleSet.h"
 #include "Widgets/Layout/SBorder.h"
+#include "Widgets/Images/SImage.h"
 #include "Widgets/Input/SButton.h"
 #include "Widgets/Text/STextBlock.h"
 #include "Widgets/Layout/SBox.h"
