- Set the icon tint in the style rather than where the image is used
  - Updated the multi-user tab icon to dim the white in it.

diff --git a/Engine/Plugins/Developer/Concert/ConcertSync/ConcertSyncClient/Source/ConcertSyncUI/Private/ConcertFrontendStyle.cpp b/Engine/Plugins/Developer/Concert/ConcertSync/ConcertSyncClient/Source/ConcertSyncUI/Private/ConcertFrontendStyle.cpp
--- a/Engine/Plugins/Developer/Concert/ConcertSync/ConcertSyncClient/Source/ConcertSyncUI/Private/ConcertFrontendStyle.cpp
+++ b/Engine/Plugins/Developer/Concert/ConcertSync/ConcertSyncClient/Source/ConcertSyncUI/Private/ConcertFrontendStyle.cpp
@@ -42,6 +42,9 @@ void FConcertFrontendStyle::Initialize()
 	const FVector2D Icon40x40(40.0f, 40.0f);
 	const FVector2D Icon48x48(48.0f, 48.0f);
 
+	// Most icons were designed to be used at 80% opacity.
+	const FLinearColor IconColorAndOpacity(FLinearColor(1.f, 1.f, 1.f, 0.8f));
+
 	// 16x16
 	StyleSet->Set("Concert.Concert", new IMAGE_PLUGIN_BRUSH("Icons/icon_Concert_16x", Icon16x16));
 	StyleSet->Set("Concert.Persist", new IMAGE_PLUGIN_BRUSH("Icons/icon_ConcertPersist_16x", Icon16x16));
@@ -54,40 +57,40 @@ void FConcertFrontendStyle::Initialize()
 
 	// Maps the UI Command name in Multi-User module. (UI_COMMAND does magic icon mapping when style name and command name matches)
 	StyleSet->Set("Concert.OpenBrowser",  new IMAGE_PLUGIN_BRUSH("Icons/icon_MultiUser_32x", Icon16x16));
-	StyleSet->Set("Concert.OpenSettings", new IMAGE_PLUGIN_BRUSH("Icons/icon_Settings_32x",  Icon16x16));
-	StyleSet->Set("Concert.LaunchServer", new IMAGE_PLUGIN_BRUSH("Icons/icon_NewServer_32x", Icon16x16));
+	StyleSet->Set("Concert.OpenSettings", new IMAGE_PLUGIN_BRUSH("Icons/icon_Settings_32x",  Icon16x16, IconColorAndOpacity));
+	StyleSet->Set("Concert.LaunchServer", new IMAGE_PLUGIN_BRUSH("Icons/icon_NewServer_32x", Icon16x16, IconColorAndOpacity));
 
 	// Multi-User Browser
-	StyleSet->Set("Concert.ArchiveSession",     new IMAGE_PLUGIN_BRUSH("Icons/icon_ArchiveSession_48x",     Icon24x24));
-	StyleSet->Set("Concert.CancelAutoJoin",     new IMAGE_PLUGIN_BRUSH("Icons/icon_CancelAutoJoin_48x",     Icon24x24));
-	StyleSet->Set("Concert.CloseServer",        new IMAGE_PLUGIN_BRUSH("Icons/icon_CloseServer_48x",        Icon24x24));
-	StyleSet->Set("Concert.DeleteSession",      new IMAGE_PLUGIN_BRUSH("Icons/icon_DeleteSession_48x",      Icon24x24));
-	StyleSet->Set("Concert.JoinDefaultSession", new IMAGE_PLUGIN_BRUSH("Icons/icon_JoinDefaultSession_48x", Icon24x24));
-	StyleSet->Set("Concert.JoinSession",        new IMAGE_PLUGIN_BRUSH("Icons/icon_JoinSelectedSession_48x",Icon24x24));
-	StyleSet->Set("Concert.LeaveSession",       new IMAGE_PLUGIN_BRUSH("Icons/icon_LeaveSession_48x",       Icon24x24));
-	StyleSet->Set("Concert.NewServer",          new IMAGE_PLUGIN_BRUSH("Icons/icon_NewServer_48x",          Icon24x24));
-	StyleSet->Set("Concert.NewSession",         new IMAGE_PLUGIN_BRUSH("Icons/icon_NewSession_48x",         Icon24x24));
-	StyleSet->Set("Concert.PauseSession",       new IMAGE_PLUGIN_BRUSH("Icons/icon_PauseSession_48x",       Icon24x24));
-	StyleSet->Set("Concert.RestoreSession",     new IMAGE_PLUGIN_BRUSH("Icons/icon_RestoreSession_48x",     Icon24x24));
-	StyleSet->Set("Concert.ResumeSession",      new IMAGE_PLUGIN_BRUSH("Icons/icon_ResumeSession_48x",      Icon24x24));
-	StyleSet->Set("Concert.Settings",           new IMAGE_PLUGIN_BRUSH("Icons/icon_Settings_48x",           Icon24x24));
-	StyleSet->Set("Concert.NewServer.Small",    new IMAGE_PLUGIN_BRUSH("Icons/icon_NewServer_32x",          Icon16x16));
-	StyleSet->Set("Concert.NewSession.Small",   new IMAGE_PLUGIN_BRUSH("Icons/icon_NewSession_32x",         Icon16x16));
+	StyleSet->Set("Concert.ArchiveSession",     new IMAGE_PLUGIN_BRUSH("Icons/icon_ArchiveSession_48x",     Icon24x24, IconColorAndOpacity));
+	StyleSet->Set("Concert.CancelAutoJoin",     new IMAGE_PLUGIN_BRUSH("Icons/icon_CancelAutoJoin_48x",     Icon24x24, IconColorAndOpacity));
+	StyleSet->Set("Concert.CloseServer",        new IMAGE_PLUGIN_BRUSH("Icons/icon_CloseServer_48x",        Icon24x24, IconColorAndOpacity));
+	StyleSet->Set("Concert.DeleteSession",      new IMAGE_PLUGIN_BRUSH("Icons/icon_DeleteSession_48x",      Icon24x24, IconColorAndOpacity));
+	StyleSet->Set("Concert.JoinDefaultSession", new IMAGE_PLUGIN_BRUSH("Icons/icon_JoinDefaultSession_48x", Icon24x24, IconColorAndOpacity));
+	StyleSet->Set("Concert.JoinSession",        new IMAGE_PLUGIN_BRUSH("Icons/icon_JoinSelectedSession_48x",Icon24x24, IconColorAndOpacity));
+	StyleSet->Set("Concert.LeaveSession",       new IMAGE_PLUGIN_BRUSH("Icons/icon_LeaveSession_48x",       Icon24x24, IconColorAndOpacity));
+	StyleSet->Set("Concert.NewServer",          new IMAGE_PLUGIN_BRUSH("Icons/icon_NewServer_48x",          Icon24x24, IconColorAndOpacity));
+	StyleSet->Set("Concert.NewSession",         new IMAGE_PLUGIN_BRUSH("Icons/icon_NewSession_48x",         Icon24x24, IconColorAndOpacity));
+	StyleSet->Set("Concert.PauseSession",       new IMAGE_PLUGIN_BRUSH("Icons/icon_PauseSession_48x",       Icon24x24, IconColorAndOpacity));
+	StyleSet->Set("Concert.RestoreSession",     new IMAGE_PLUGIN_BRUSH("Icons/icon_RestoreSession_48x",     Icon24x24, IconColorAndOpacity));
+	StyleSet->Set("Concert.ResumeSession",      new IMAGE_PLUGIN_BRUSH("Icons/icon_ResumeSession_48x",      Icon24x24, IconColorAndOpacity));
+	StyleSet->Set("Concert.Settings",           new IMAGE_PLUGIN_BRUSH("Icons/icon_Settings_48x",           Icon24x24, IconColorAndOpacity));
+	StyleSet->Set("Concert.NewServer.Small",    new IMAGE_PLUGIN_BRUSH("Icons/icon_NewServer_32x",          Icon16x16, IconColorAndOpacity));
+	StyleSet->Set("Concert.NewSession.Small",   new IMAGE_PLUGIN_BRUSH("Icons/icon_NewSession_32x",         Icon16x16, IconColorAndOpacity));
 
 	// Multi-user Active session
-	StyleSet->Set("Concert.JumpToLocation",     new IMAGE_PLUGIN_BRUSH("Icons/icon_PresenceLocation_32x",   Icon16x16));
-	StyleSet->Set("Concert.HidePresence",       new IMAGE_PLUGIN_BRUSH("Icons/icon_PresenceEyeOff_32x",     Icon16x16));
-	StyleSet->Set("Concert.ShowPresence",       new IMAGE_PLUGIN_BRUSH("Icons/icon_PresenceEyeOn_32x",      Icon16x16));
+	StyleSet->Set("Concert.JumpToLocation",     new IMAGE_PLUGIN_BRUSH("Icons/icon_PresenceLocation_32x",   Icon16x16, IconColorAndOpacity));
+	StyleSet->Set("Concert.HidePresence",       new IMAGE_PLUGIN_BRUSH("Icons/icon_PresenceEyeOff_32x",     Icon16x16, IconColorAndOpacity));
+	StyleSet->Set("Concert.ShowPresence",       new IMAGE_PLUGIN_BRUSH("Icons/icon_PresenceEyeOn_32x",      Icon16x16, IconColorAndOpacity));
 
 	// 20x20 -> Editor toolbar small icons.
 	StyleSet->Set("Concert.Online.Small", new IMAGE_PLUGIN_BRUSH("Icons/icon_ConcertOnline_40x", Icon20x20));
 	StyleSet->Set("Concert.Offline.Small", new IMAGE_PLUGIN_BRUSH("Icons/icon_ConcertOffline_40x", Icon20x20));
 
 	// 24x24/48x48 -> For sequencer toolbar.
-	StyleSet->Set("Concert.Sequencer.SyncTimeline",       new IMAGE_PLUGIN_BRUSH("Icons/icon_SequencerSyncTimeline_48x", Icon48x48)); // Enable/disable playback and time scrubbing from a remote client.
-	StyleSet->Set("Concert.Sequencer.SyncTimeline.Small", new IMAGE_PLUGIN_BRUSH("Icons/icon_SequencerSyncTimeline_48x", Icon24x24));
-	StyleSet->Set("Concert.Sequencer.SyncSequence",       new IMAGE_PLUGIN_BRUSH("Icons/icon_SequencerSyncSequence_48x", Icon48x48)); // Allows or not a remote client to open/close sequencer.
-	StyleSet->Set("Concert.Sequencer.SyncSequence.Small", new IMAGE_PLUGIN_BRUSH("Icons/icon_SequencerSyncSequence_48x", Icon24x24));
+	StyleSet->Set("Concert.Sequencer.SyncTimeline",       new IMAGE_PLUGIN_BRUSH("Icons/icon_SequencerSyncTimeline_48x", Icon48x48, IconColorAndOpacity)); // Enable/disable playback and time scrubbing from a remote client.
+	StyleSet->Set("Concert.Sequencer.SyncTimeline.Small", new IMAGE_PLUGIN_BRUSH("Icons/icon_SequencerSyncTimeline_48x", Icon24x24, IconColorAndOpacity));
+	StyleSet->Set("Concert.Sequencer.SyncSequence",       new IMAGE_PLUGIN_BRUSH("Icons/icon_SequencerSyncSequence_48x", Icon48x48, IconColorAndOpacity)); // Allows or not a remote client to open/close sequencer.
+	StyleSet->Set("Concert.Sequencer.SyncSequence.Small", new IMAGE_PLUGIN_BRUSH("Icons/icon_SequencerSyncSequence_48x", Icon24x24, IconColorAndOpacity));
 
 	// 40x40 -> Editor toolbar large icons.
 	StyleSet->Set("Concert.Online", new IMAGE_PLUGIN_BRUSH("Icons/icon_ConcertOnline_40x", Icon40x40));
@@ -103,7 +106,6 @@ void FConcertFrontendStyle::Initialize()
 	{
 		StyleSet->Set("Concert.Color.LocalUser", FLinearColor(0.31f, 0.749f, 0.333f));
 		StyleSet->Set("Concert.Color.OtherUser", FLinearColor(0.93f, 0.608f, 0.169f));
-		StyleSet->Set("Concert.Color.ToolbarButtonIcon", FLinearColor::White * FLinearColor(1.0f, 1.0f, 1.0f, 0.8f)); // Icons are made to have 80% opacity.
 	}
 
 	// Colors
diff --git a/Engine/Plugins/Developer/Concert/ConcertSync/ConcertSyncClient/Source/ConcertSyncUI/Public/ConcertFrontendUtils.h b/Engine/Plugins/Developer/Concert/ConcertSync/ConcertSyncClient/Source/ConcertSyncUI/Public/ConcertFrontendUtils.h
--- a/Engine/Plugins/Developer/Concert/ConcertSync/ConcertSyncClient/Source/ConcertSyncUI/Public/ConcertFrontendUtils.h
+++ b/Engine/Plugins/Developer/Concert/ConcertSync/ConcertSyncClient/Source/ConcertSyncUI/Public/ConcertFrontendUtils.h
@@ -94,7 +94,6 @@ namespace ConcertFrontendUtils
 			[
 				SNew(SImage)
 				.Image(TAttribute<const FSlateBrush*>::Create([IconStyleAttr = InDef.IconStyle]() { return FConcertFrontendStyle::Get()->GetBrush(IconStyleAttr.Get()); }))
-				.ColorAndOpacity(FConcertFrontendStyle::Get()->GetColor("Concert.Color.ToolbarButtonIcon"))
 			];
 	}
 
diff --git a/Engine/Plugins/Developer/Concert/MultiUserClient/Source/MultiUserClient/Private/Widgets/SConcertBrowser.cpp b/Engine/Plugins/Developer/Concert/MultiUserClient/Source/MultiUserClient/Private/Widgets/SConcertBrowser.cpp
--- a/Engine/Plugins/Developer/Concert/MultiUserClient/Source/MultiUserClient/Private/Widgets/SConcertBrowser.cpp
+++ b/Engine/Plugins/Developer/Concert/MultiUserClient/Source/MultiUserClient/Private/Widgets/SConcertBrowser.cpp
@@ -99,7 +99,7 @@ TSharedRef<SButton> MakeIconButton(const FName& ButtonStyle, const TAttribute<co
 		.VAlign(VAlign_Center)
 		.HAlign(HAlign_Center)
 		[
-			SNew(SImage).Image(Icon).ColorAndOpacity(FConcertFrontendStyle::Get()->GetColor("Concert.Color.ToolbarButtonIcon"))
+			SNew(SImage).Image(Icon)
 		];
 }
 
@@ -1705,7 +1705,7 @@ public:
 						.VAlign(VAlign_Center)
 						.Padding(0, 0, 3, 0)
 						[
-							SNew(SImage).Image(InArgs._ButtonIcon).ColorAndOpacity(FConcertFrontendStyle::Get()->GetColor("Concert.Color.ToolbarButtonIcon"))
+							SNew(SImage).Image(InArgs._ButtonIcon)
 						]
 						+SHorizontalBox::Slot()
 						.AutoWidth()
