Editor fix to unblock CIS

diff --git a/Engine/Plugins/Developer/Concert/ConcertSync/ConcertSyncClient/Source/ConcertSyncClient/Private/ConcertClientSequencerManager.cpp b/Engine/Plugins/Developer/Concert/ConcertSync/ConcertSyncClient/Source/ConcertSyncClient/Private/ConcertClientSequencerManager.cpp
--- a/Engine/Plugins/Developer/Concert/ConcertSync/ConcertSyncClient/Source/ConcertSyncClient/Private/ConcertClientSequencerManager.cpp
+++ b/Engine/Plugins/Developer/Concert/ConcertSync/ConcertSyncClient/Source/ConcertSyncClient/Private/ConcertClientSequencerManager.cpp
@@ -438,7 +438,7 @@ void FSequencerEventClient::ApplyEventToPlayers(const FConcertSequencerState& Ev
 		if (Sequence && CurrentWorld)
 		{
 			Player = NewObject<ULevelSequencePlayer>((UObject*)GetTransientPackage(), FName("ConcertSequencePlayer"));
-			Player->Initialize(Sequence, CurrentWorld, FMovieSceneSequencePlaybackSettings());
+			Player->Initialize(Sequence, CurrentWorld->GetLevel(0), FMovieSceneSequencePlaybackSettings());
 		}
 		SequencePlayers.Add(*EventState.SequenceObjectPath, Player);
 	}
