Do not skip assets when seriliazing asset reference for annotations in Concert

diff --git a/Engine/Plugins/Developer/Concert/ConcertSync/ConcertSyncClient/Source/ConcertSyncClient/Private/ConcertClientTransactionManager.cpp b/Engine/Plugins/Developer/Concert/ConcertSync/ConcertSyncClient/Source/ConcertSyncClient/Private/ConcertClientTransactionManager.cpp
--- a/Engine/Plugins/Developer/Concert/ConcertSync/ConcertSyncClient/Source/ConcertSyncClient/Private/ConcertClientTransactionManager.cpp
+++ b/Engine/Plugins/Developer/Concert/ConcertSync/ConcertSyncClient/Source/ConcertSyncClient/Private/ConcertClientTransactionManager.cpp
@@ -354,7 +354,7 @@ void FConcertClientTransactionManager::HandleObjectTransacted(UObject* InObject,
 
 		if (TransactionAnnotation.IsValid())
 		{
-			FConcertSyncObjectWriter AnnotationWriter(&PendingTransaction.FinalizedLocalIdentifierTable, InObject, ObjectUpdate.SerializedAnnotationData, bIncludeEditorOnlyProperties, true);
+			FConcertSyncObjectWriter AnnotationWriter(&PendingTransaction.FinalizedLocalIdentifierTable, InObject, ObjectUpdate.SerializedAnnotationData, bIncludeEditorOnlyProperties, false);
 			TransactionAnnotation->Serialize(AnnotationWriter);
 		}
 
