Renamed source control proxy from Concert to Multi-User

diff --git a/Engine/Plugins/Developer/Concert/ConcertSync/ConcertSyncClient/Source/ConcertSyncClient/Private/ConcertSourceControlProxy.cpp b/Engine/Plugins/Developer/Concert/ConcertSync/ConcertSyncClient/Source/ConcertSyncClient/Private/ConcertSourceControlProxy.cpp
--- a/Engine/Plugins/Developer/Concert/ConcertSync/ConcertSyncClient/Source/ConcertSyncClient/Private/ConcertSourceControlProxy.cpp
+++ b/Engine/Plugins/Developer/Concert/ConcertSync/ConcertSyncClient/Source/ConcertSyncClient/Private/ConcertSourceControlProxy.cpp
@@ -10,7 +10,7 @@
 
 #include "ConcertClientWorkspace.h"
 
-const FName FConcertSourceControlProxy::ConcertProviderName("Concert");
+const FName FConcertSourceControlProxy::ConcertProviderName("Multi-User");
 
 #define LOCTEXT_NAMESPACE "ConcertSourceControl"
 
