Remove redundent bCanEverTick sets for deprecated PostPhysicsComponentTick
#rnx

diff --git a/Engine/Plugins/Developer/Concert/ConcertSync/ConcertSyncClient/Source/ConcertSyncClient/Private/ConcertClientVRPresenceActor.cpp b/Engine/Plugins/Developer/Concert/ConcertSync/ConcertSyncClient/Source/ConcertSyncClient/Private/ConcertClientVRPresenceActor.cpp
--- a/Engine/Plugins/Developer/Concert/ConcertSync/ConcertSyncClient/Source/ConcertSyncClient/Private/ConcertClientVRPresenceActor.cpp
+++ b/Engine/Plugins/Developer/Concert/ConcertSync/ConcertSyncClient/Source/ConcertSyncClient/Private/ConcertClientVRPresenceActor.cpp
@@ -36,7 +36,6 @@ AConcertClientVRPresenceActor::AConcertClientVRPresenceActor(const FObjectInitia
 	AddOwnedComponent(LaserSplineComponent);
 	LaserSplineComponent->SetupAttachment(RootComponent);
 	LaserSplineComponent->SetVisibility(false);
-	LaserSplineComponent->PostPhysicsComponentTick.bCanEverTick = false;
 
 	bIsRightControllerVisible = true;
 	bIsLeftControllerVisible = true;
@@ -209,7 +208,6 @@ void AConcertClientVRPresenceActor::InitPresence(const class UConcertAssetContai
 		SplineSegment->SetMobility(EComponentMobility::Movable);
 		SplineSegment->SetCollisionEnabled(ECollisionEnabled::NoCollision);
 		SplineSegment->SetSplineUpDir(FVector::UpVector, false);
-		SplineSegment->PostPhysicsComponentTick.bCanEverTick = false;
 
 		UStaticMesh* StaticMesh = nullptr;
 		if (i == 0)
diff --git a/Engine/Source/Editor/VREditor/VREditorInteractor.cpp b/Engine/Source/Editor/VREditor/VREditorInteractor.cpp
--- a/Engine/Source/Editor/VREditor/VREditorInteractor.cpp
+++ b/Engine/Source/Editor/VREditor/VREditorInteractor.cpp
@@ -324,7 +324,6 @@ void UVREditorInteractor::SetupComponent_Implementation( AActor* OwningActor )
 		LaserSplineComponent->SetupAttachment( MotionControllerComponent );
 		LaserSplineComponent->RegisterComponent();
 		LaserSplineComponent->SetVisibility( false );
-		LaserSplineComponent->PostPhysicsComponentTick.bCanEverTick = false;
 		LaserSplineMeshComponents.Empty();
 
 		for (int32 i = 0; i < NumLaserSplinePoints; i++)
@@ -333,7 +332,6 @@ void UVREditorInteractor::SetupComponent_Implementation( AActor* OwningActor )
 			SplineSegment->SetMobility( EComponentMobility::Movable );
 			SplineSegment->SetCollisionEnabled( ECollisionEnabled::NoCollision );
 			SplineSegment->SetSplineUpDir( FVector::UpVector, false );
-			SplineSegment->PostPhysicsComponentTick.bCanEverTick = false;
 
 			UStaticMesh* StaticMesh = nullptr;
 			if (i == 0)
