When a session is left, the multi-user system performs some clean up to remove 'session only' transactions. If a level was created in a multi-users session and the session left while that world package was dirty, the system would revert those transactions and that process involve deleting/reloading the map. The delete action was unexpectedly recorded by the multi-user system as transaction (map deleted) because the sytem was still hooked to the editor events when this happen. Unregistring the multi-user calls back before performing the clean up fixed the issue.

When a level that was created in a session (and only exists within that session) was the active world when the user left the session, the editor would not unload it. The level could be edited, but could not be saved. Reset the world to the default level template when this situation is detected.

diff --git a/Engine/Plugins/Developer/Concert/ConcertSync/ConcertSyncClient/Source/ConcertSyncClient/ConcertSyncClient.Build.cs b/Engine/Plugins/Developer/Concert/ConcertSync/ConcertSyncClient/Source/ConcertSyncClient/ConcertSyncClient.Build.cs
--- a/Engine/Plugins/Developer/Concert/ConcertSync/ConcertSyncClient/Source/ConcertSyncClient/ConcertSyncClient.Build.cs
+++ b/Engine/Plugins/Developer/Concert/ConcertSync/ConcertSyncClient/Source/ConcertSyncClient/ConcertSyncClient.Build.cs
@@ -55,6 +55,7 @@ namespace UnrealBuildTool.Rules
 				PrivateDependencyModuleNames.AddRange(
 					new string[] {
 						"EditorStyle",
+						"EngineSettings",
 						"Sequencer",
 						"UnrealEd",
 						"ViewportInteraction",
diff --git a/Engine/Plugins/Developer/Concert/ConcertSync/ConcertSyncClient/Source/ConcertSyncClient/Private/ConcertClientWorkspace.cpp b/Engine/Plugins/Developer/Concert/ConcertSync/ConcertSyncClient/Source/ConcertSyncClient/Private/ConcertClientWorkspace.cpp
--- a/Engine/Plugins/Developer/Concert/ConcertSync/ConcertSyncClient/Source/ConcertSyncClient/Private/ConcertClientWorkspace.cpp
+++ b/Engine/Plugins/Developer/Concert/ConcertSync/ConcertSyncClient/Source/ConcertSyncClient/Private/ConcertClientWorkspace.cpp
@@ -46,6 +46,8 @@
 	#include "Editor/UnrealEdEngine.h"
 	#include "Editor/TransBuffer.h"
 	#include "LevelEditor.h"
+	#include "FileHelpers.h"
+	#include "GameMapsSettings.h"
 #endif
 
 #define LOCTEXT_NAMESPACE "ConcertClientWorkspace"
@@ -309,12 +311,6 @@ void FConcertClientWorkspace::UnbindSession()
 
 		// Gather file with live transactions that also need to be reloaded, overlaps from the sandbox are filtered directly in ReloadPackages
 		PackagesPendingHotReload.Append(TransactionManager->GetLedger().GetPackagesNamesWithLiveTransactions());
-
-		if (!GIsRequestingExit)
-		{
-			HotReloadPendingPackages();
-			PurgePendingPackages();
-		}
 #endif
 
 		// Destroy Transaction Manager
@@ -399,6 +395,33 @@ void FConcertClientWorkspace::UnbindSession()
 			FCoreUObjectDelegates::OnObjectTransacted.Remove(ObjectTransactedHandle);
 			ObjectTransactedHandle.Reset();
 		}
+
+		if (!GIsRequestingExit)
+		{
+			// Hot reload after unregistering from most delegates to prevent events triggered by hot-reloading (such as asset deleted) to be recorded as transaction.
+			HotReloadPendingPackages();
+
+			// Get the current world edited.
+			if (UWorld* World = GEditor->GetEditorWorldContext().World())
+			{
+				// If the current world package is scheduled to be purged (it doesn't exist outside the session).
+				if (PackagesPendingPurge.Contains(World->GetOutermost()->GetFName()))
+				{
+					// Replace the current world because it doesn't exist outside the session (it cannot be saved anymore, even with 'Save Current As').
+					FString StartupMapPackage = GetDefault<UGameMapsSettings>()->EditorStartupMap.GetLongPackageName();
+					if (FPackageName::DoesPackageExist(StartupMapPackage))
+					{
+						UEditorLoadingAndSavingUtils::NewMapFromTemplate(StartupMapPackage, /*bSaveExistingMap*/ false);
+					}
+					else
+					{
+						UEditorLoadingAndSavingUtils::NewBlankMap(/*bSaveExistingMap*/ false);
+					}
+				}
+
+				PurgePendingPackages();
+			}
+		}
 #endif
 
 		// Unregister OnEndFrame events
