Fixed some MU presence issues

 - Fixed presence avatar settings being empty/incorrect.
 - Fixed base presence types being available for selection in the settings.
 - Fixed presence manager data being susceptible to GC.
 - Fixed presence manager defaulting to VR mode.
 - Removed old presence assets.

#rnx

diff --git a/Engine/Plugins/Developer/Concert/ConcertMain/Source/Concert/Public/ConcertSettings.h b/Engine/Plugins/Developer/Concert/ConcertMain/Source/Concert/Public/ConcertSettings.h
--- a/Engine/Plugins/Developer/Concert/ConcertMain/Source/Concert/Public/ConcertSettings.h
+++ b/Engine/Plugins/Developer/Concert/ConcertMain/Source/Concert/Public/ConcertSettings.h
@@ -161,8 +161,8 @@ struct FConcertClientSettings
 	FConcertClientSettings()
 		: DisplayName()
 		, AvatarColor(1.0f, 1.0f, 1.0f, 1.0f)
-		, DesktopAvatarActorClass(TEXT("/ConcertSyncClient/ConcertSyncClient.ConcertClientDesktopPresenceActor"))
-		, VRAvatarActorClass(TEXT("/ConcertSyncClient/ConcertSyncClient.ConcertClientVRPresenceActor"))
+		, DesktopAvatarActorClass(TEXT("/ConcertSyncClient/DesktopPresence.DesktopPresence_C"))
+		, VRAvatarActorClass(TEXT("/ConcertSyncClient/VRPresence.VRPresence_C"))
 		, DiscoveryTimeoutSeconds(5)
 		, SessionTickFrequencySeconds(1)
 		, LatencyCompensationMs(0)
@@ -180,11 +180,11 @@ struct FConcertClientSettings
 	FLinearColor AvatarColor;
 
 	/** The desktop representation of this editor's user to other connected users */
-	UPROPERTY(config, EditAnywhere, NoClear, Category = "Client Settings", meta = (MetaClass = "ConcertClientPresenceActor"))
+	UPROPERTY(config, EditAnywhere, NoClear, Category = "Client Settings", meta = (MetaClass = "ConcertClientDesktopPresenceActor"))
 	FSoftClassPath DesktopAvatarActorClass;
 
 	/** The VR representation of this editor's user to other connected users */
-	UPROPERTY(config, EditAnywhere, NoClear, Category = "Client Settings", meta = (MetaClass = "ConcertClientPresenceActor"))
+	UPROPERTY(config, EditAnywhere, NoClear, Category = "Client Settings", meta = (MetaClass = "ConcertClientVRPresenceActor"))
 	FSoftClassPath VRAvatarActorClass;
 
 	/** The timespan at which discovered Concert server are considered stale if they haven't answered back */
diff --git a/Engine/Plugins/Developer/Concert/ConcertSync/ConcertSyncClient/Source/ConcertSyncClient/Private/ConcertClientDesktopPresenceActor.h b/Engine/Plugins/Developer/Concert/ConcertSync/ConcertSyncClient/Source/ConcertSyncClient/Private/ConcertClientDesktopPresenceActor.h
--- a/Engine/Plugins/Developer/Concert/ConcertSync/ConcertSyncClient/Source/ConcertSyncClient/Private/ConcertClientDesktopPresenceActor.h
+++ b/Engine/Plugins/Developer/Concert/ConcertSync/ConcertSyncClient/Source/ConcertSyncClient/Private/ConcertClientDesktopPresenceActor.h
@@ -14,7 +14,7 @@ struct FDesktopPresenceUpdateEvent;
 /**
   * A AConcertClientDesktopPresenceActor is a child of AConcertClientPresenceActor that is used to represent users in desktop 
   */
-UCLASS(Transient, NotPlaceable, Blueprintable)
+UCLASS(Abstract, Transient, NotPlaceable, Blueprintable)
 class AConcertClientDesktopPresenceActor : public AConcertClientPresenceActor
 {
 	GENERATED_UCLASS_BODY()
diff --git a/Engine/Plugins/Developer/Concert/ConcertSync/ConcertSyncClient/Source/ConcertSyncClient/Private/ConcertClientPresenceActor.h b/Engine/Plugins/Developer/Concert/ConcertSync/ConcertSyncClient/Source/ConcertSyncClient/Private/ConcertClientPresenceActor.h
--- a/Engine/Plugins/Developer/Concert/ConcertSync/ConcertSyncClient/Source/ConcertSyncClient/Private/ConcertClientPresenceActor.h
+++ b/Engine/Plugins/Developer/Concert/ConcertSync/ConcertSyncClient/Source/ConcertSyncClient/Private/ConcertClientPresenceActor.h
@@ -18,7 +18,7 @@ class FConcertClientPresenceMode;
 /**
   * A ConcertClientPresenceActor is an editor-only transient actor representing other client presences during a concert client session.
   */
-UCLASS(Transient, NotPlaceable, Blueprintable)
+UCLASS(Abstract, Transient, NotPlaceable, Blueprintable)
 class AConcertClientPresenceActor : public AActor
 {
 	GENERATED_UCLASS_BODY()
diff --git a/Engine/Plugins/Developer/Concert/ConcertSync/ConcertSyncClient/Source/ConcertSyncClient/Private/ConcertClientPresenceManager.cpp b/Engine/Plugins/Developer/Concert/ConcertSync/ConcertSyncClient/Source/ConcertSyncClient/Private/ConcertClientPresenceManager.cpp
--- a/Engine/Plugins/Developer/Concert/ConcertSync/ConcertSyncClient/Source/ConcertSyncClient/Private/ConcertClientPresenceManager.cpp
+++ b/Engine/Plugins/Developer/Concert/ConcertSync/ConcertSyncClient/Source/ConcertSyncClient/Private/ConcertClientPresenceManager.cpp
@@ -57,11 +57,10 @@ bool ShowPresenceInPIE(const bool InIsPIE)
 
 }
 
-static TAutoConsoleVariable<int32> CVarEnablePresence(TEXT("concert.EnablePresence"), 1, TEXT("Enable Concert presence"));
+static TAutoConsoleVariable<int32> CVarEnablePresence(TEXT("concert.EnablePresence"), 1, TEXT("Enable Concert Presence"));
 
 const TCHAR* FConcertClientPresenceManager::AssetContainerPath = TEXT("/ConcertSyncClient/ConcertAssets");
 
-
 FConcertClientPresenceManager::FConcertClientPresenceManager(TSharedRef<IConcertClientSession> InSession)
 	: OnSessionClientChangedHandle()
 	, OnVREditingModeEnterHandle()
@@ -70,13 +69,14 @@ FConcertClientPresenceManager::FConcertClientPresenceManager(TSharedRef<IConcert
 	, CurrentAvatarMode(nullptr)
 	, AssetContainer(nullptr)
 	, bIsPresenceEnabled(true)
+	, bInVR(false)
 	, CurrentAvatarActorClass(nullptr)
 	, DesktopAvatarActorClass(nullptr)
 	, VRAvatarActorClass(nullptr)
 {
 	// Setup the asset container.
-	AssetContainer = &LoadAssetContainer();
-	check(AssetContainer);
+	AssetContainer = LoadObject<UConcertAssetContainer>(nullptr, FConcertClientPresenceManager::AssetContainerPath);
+	checkf(AssetContainer, TEXT("Failed to load UConcertAssetContainer (%s). See log for reason."), FConcertClientPresenceManager::AssetContainerPath);
 
 	// @todo - Need to handle the situation where the avatar class might change during a session.
 	// This makes the assumption that avatar class will not change during a session 
@@ -102,6 +102,14 @@ FConcertClientPresenceManager::~FConcertClientPresenceManager()
 	ClearAllPresenceState();
 }
 
+void FConcertClientPresenceManager::AddReferencedObjects(FReferenceCollector& Collector)
+{
+	Collector.AddReferencedObject(AssetContainer);
+	Collector.AddReferencedObject(CurrentAvatarActorClass);
+	Collector.AddReferencedObject(DesktopAvatarActorClass);
+	Collector.AddReferencedObject(VRAvatarActorClass);
+}
+
 const UConcertAssetContainer& FConcertClientPresenceManager::GetAssetContainer() const
 {
 	return *AssetContainer;
@@ -560,21 +568,6 @@ void FConcertClientPresenceManager::UpdatePresenceAvatar(const FGuid& InEndpoint
 	}
 }
 
-UConcertAssetContainer& FConcertClientPresenceManager::LoadAssetContainer()
-{
-	UConcertAssetContainer* AssetContainer = LoadObject<UConcertAssetContainer>(nullptr, FConcertClientPresenceManager::AssetContainerPath);
-	checkf(AssetContainer, TEXT("Failed to load UConcertAssetContainer (%s). See log for reason."), FConcertClientPresenceManager::AssetContainerPath);
-	return *AssetContainer;
-}
-
-void FConcertClientPresenceManager::AddReferencedObjects(FReferenceCollector& Collector)
-{
-	if (AssetContainer)
-	{
-		Collector.AddReferencedObject(AssetContainer);
-	}
-}
-
 void FConcertClientPresenceManager::SetPresenceInPIE(const FGuid& InEndpointId, bool bPIE)
 {
 	FConcertClientPresenceState& PresenceState = EnsurePresenceState(InEndpointId);
diff --git a/Engine/Plugins/Developer/Concert/ConcertSync/ConcertSyncClient/Source/ConcertSyncClient/Private/ConcertClientPresenceManager.h b/Engine/Plugins/Developer/Concert/ConcertSync/ConcertSyncClient/Source/ConcertSyncClient/Private/ConcertClientPresenceManager.h
--- a/Engine/Plugins/Developer/Concert/ConcertSync/ConcertSyncClient/Source/ConcertSyncClient/Private/ConcertClientPresenceManager.h
+++ b/Engine/Plugins/Developer/Concert/ConcertSync/ConcertSyncClient/Source/ConcertSyncClient/Private/ConcertClientPresenceManager.h
@@ -11,7 +11,7 @@
 #include "ConcertClientPresenceActor.h"
 #include "ConcertClientPresenceMode.h"
 #include "UObject/Class.h"
-#include "UObject/StrongObjectPtr.h"
+#include "UObject/GCObject.h"
 
 #if WITH_EDITOR
 
@@ -78,12 +78,15 @@ struct FConcertClientPresencePersistentState
 	bool bPropagateToAll;
 };
 
-class FConcertClientPresenceManager : public TSharedFromThis<FConcertClientPresenceManager>
+class FConcertClientPresenceManager : public TSharedFromThis<FConcertClientPresenceManager>, public FGCObject
 {
 public:
 	FConcertClientPresenceManager(TSharedRef<IConcertClientSession> InSession);
 	~FConcertClientPresenceManager();
 
+	//~ FGCObject interfaces
+	virtual void AddReferencedObjects(FReferenceCollector& Collector) override;
+
 	/** Gets the container for all the assets of Concert clients. */
 	const class UConcertAssetContainer& GetAssetContainer() const;
 
@@ -197,12 +200,6 @@ private:
 	/** Updates presence avatar for remote client by invalidating current presence actor */
 	void UpdatePresenceAvatar(const FGuid& InEndpointId, bool bIsInVR);
 
-	/** Loads and returns the container for all the assets of Concert clients. */
-	static class UConcertAssetContainer& LoadAssetContainer();
-
-	/** Need to manage the GC of the asset container */
-	void AddReferencedObjects(FReferenceCollector& Collector);
-
 	/** Set presence PIE state */
 	void SetPresenceInPIE(const FGuid& InEndpointId, bool bInPIE);
 
diff --git a/Engine/Plugins/Developer/Concert/ConcertSync/ConcertSyncClient/Source/ConcertSyncClient/Private/ConcertClientVRPresenceActor.h b/Engine/Plugins/Developer/Concert/ConcertSync/ConcertSyncClient/Source/ConcertSyncClient/Private/ConcertClientVRPresenceActor.h
--- a/Engine/Plugins/Developer/Concert/ConcertSync/ConcertSyncClient/Source/ConcertSyncClient/Private/ConcertClientVRPresenceActor.h
+++ b/Engine/Plugins/Developer/Concert/ConcertSync/ConcertSyncClient/Source/ConcertSyncClient/Private/ConcertClientVRPresenceActor.h
@@ -11,7 +11,7 @@
 /**
   * A ConcertClientVRPresenceActor is a child of AConcertClientPresenceActor that is used to represent users in VR 
   */
-UCLASS(Transient, NotPlaceable, Blueprintable)
+UCLASS(Abstract, Transient, NotPlaceable, Blueprintable)
 class AConcertClientVRPresenceActor : public AConcertClientPresenceActor
 {
 	GENERATED_UCLASS_BODY()
