Fixed hot-reload error that could happen when leaving a MU session

#rnx

diff --git a/Engine/Plugins/Developer/Concert/ConcertSync/ConcertSyncClient/Source/ConcertSyncClient/Private/ConcertClientWorkspace.cpp b/Engine/Plugins/Developer/Concert/ConcertSync/ConcertSyncClient/Source/ConcertSyncClient/Private/ConcertClientWorkspace.cpp
--- a/Engine/Plugins/Developer/Concert/ConcertSync/ConcertSyncClient/Source/ConcertSyncClient/Private/ConcertClientWorkspace.cpp
+++ b/Engine/Plugins/Developer/Concert/ConcertSync/ConcertSyncClient/Source/ConcertSyncClient/Private/ConcertClientWorkspace.cpp
@@ -310,7 +310,13 @@ void FConcertClientWorkspace::UnbindSession()
 		SandboxPlatformFile.Reset();
 
 		// Gather file with live transactions that also need to be reloaded, overlaps from the sandbox are filtered directly in ReloadPackages
-		PackagesPendingHotReload.Append(TransactionManager->GetLedger().GetPackagesNamesWithLiveTransactions());
+		for (const FName PackageNameWithLiveTransactions : TransactionManager->GetLedger().GetPackagesNamesWithLiveTransactions())
+		{
+			if (!PackagesPendingPurge.Contains(PackageNameWithLiveTransactions))
+			{
+				PackagesPendingHotReload.Add(PackageNameWithLiveTransactions);
+			}
+		}
 #endif
 
 		// Destroy Transaction Manager
