Avoid loading packages unless needed
This change goes back to using Find rather than Load in most cases, to avoid always loading assets into memory when they're changed live in a multi-user environment (which could cause hitches).
This had been changed a while back to always load packages, as there was a case that could cause a dummy package to be created in-memory. This issue could not be reproduced.

diff --git a/Engine/Plugins/Developer/Concert/ConcertSync/ConcertSyncClient/Source/ConcertSyncClient/Private/ConcertSyncClientUtil.cpp b/Engine/Plugins/Developer/Concert/ConcertSync/ConcertSyncClient/Source/ConcertSyncClient/Private/ConcertSyncClientUtil.cpp
--- a/Engine/Plugins/Developer/Concert/ConcertSync/ConcertSyncClient/Source/ConcertSyncClient/Private/ConcertSyncClientUtil.cpp
+++ b/Engine/Plugins/Developer/Concert/ConcertSync/ConcertSyncClient/Source/ConcertSyncClient/Private/ConcertSyncClientUtil.cpp
@@ -68,31 +68,20 @@ FGetObjectResult GetObject(const FConcertObjectId& InObjectId, const FName InNew
 	const FName ObjectNameToFind = InObjectId.ObjectName;
 	const FName ObjectNameToCreate = bIsRename ? InNewName : ObjectNameToFind;
 
-	auto FindOrLoadOuterPackage = [](const FName InOuterName) -> UObject*
+	auto FindOrLoadClass = [bAllowCreate](const FName InClassName) -> UClass*
 	{
-		const FString OuterNameStr = InOuterName.ToString();
+		const FString ClassNameStr = InClassName.ToString();
 
-		// Always attempt to find an in-memory object first as we may be calling this function while a load is taking place
-		if (UObject* FoundOuter = StaticFindObject(UObject::StaticClass(), nullptr, *OuterNameStr))
-		{
-			return FoundOuter;
-		}
-
-		// If the outer name is a package path that isn't currently loaded, then we need to try loading it to avoid 
-		// creating an in-memory version of the package (which would prevent the real package ever loading)
-		if (FPackageName::IsValidLongPackageName(OuterNameStr))
-		{
-			return LoadPackage(nullptr, *OuterNameStr, LOAD_NoWarn);
-		}
-
-		return nullptr;
+		return bAllowCreate
+			? LoadObject<UClass>(nullptr, *ClassNameStr)
+			: FindObject<UClass>(nullptr, *ClassNameStr);
 	};
 
-	// We need the object class to do anything
-	if (UClass* ObjectClass = LoadObject<UClass>(nullptr, *InObjectId.ObjectClassPathName.ToString()))
+	// Find the outer for the existing object
+	if (UObject* ExistingObjectOuter = StaticFindObject(UObject::StaticClass(), nullptr, *ObjectOuterPathToFind.ToString()))
 	{
-		// Find the outer for the existing object
-		if (UObject* ExistingObjectOuter = FindOrLoadOuterPackage(ObjectOuterPathToFind))
+		// We need the object class to find or create the object
+		if (UClass* ObjectClass = FindOrLoadClass(InObjectId.ObjectClassPathName))
 		{
 			// Find the existing object
 			if (UObject* ExistingObject = StaticFindObject(ObjectClass, ExistingObjectOuter, *ObjectNameToFind.ToString(), /*bExactClass*/true))
@@ -130,9 +119,13 @@ FGetObjectResult GetObject(const FConcertObjectId& InObjectId, const FName InNew
 				return FGetObjectResult(ExistingObject, ResultFlags);
 			}
 		}
+	}
 
-		// Find the outer for the new object
-		if (UObject* NewObjectOuter = FindOrLoadOuterPackage(ObjectOuterPathToCreate))
+	// Find the outer for the new object
+	if (UObject* NewObjectOuter = StaticFindObject(UObject::StaticClass(), nullptr, *ObjectOuterPathToCreate.ToString()))
+	{
+		// We need the object class to find or create the object
+		if (UClass* ObjectClass = FindOrLoadClass(InObjectId.ObjectClassPathName))
 		{
 			// Find the new object (in case something already created it)
 			if (UObject* NewObject = StaticFindObject(ObjectClass, NewObjectOuter, *ObjectNameToCreate.ToString(), /*bExactClass*/true))
