Avoid applying transactions multiple times when loading worlds

diff --git a/Engine/Plugins/Developer/Concert/ConcertSync/ConcertSyncClient/Source/ConcertSyncClient/Private/ConcertClientWorkspace.cpp b/Engine/Plugins/Developer/Concert/ConcertSync/ConcertSyncClient/Source/ConcertSyncClient/Private/ConcertClientWorkspace.cpp
--- a/Engine/Plugins/Developer/Concert/ConcertSync/ConcertSyncClient/Source/ConcertSyncClient/Private/ConcertClientWorkspace.cpp
+++ b/Engine/Plugins/Developer/Concert/ConcertSync/ConcertSyncClient/Source/ConcertSyncClient/Private/ConcertClientWorkspace.cpp
@@ -270,13 +270,6 @@ void FConcertClientWorkspace::BindSession(const TSharedRef<IConcertClientSession
 	SwitchBeginPIEAndSIEHandle = FEditorDelegates::OnSwitchBeginPIEAndSIE.AddRaw(this, &FConcertClientWorkspace::HandleSwitchBeginPIEAndSIE);
 	EndPIEHandle = FEditorDelegates::EndPIE.AddRaw(this, &FConcertClientWorkspace::HandleEndPIE);
 
-	// Register Map Change Events
-	if (GIsEditor)
-	{
-		FLevelEditorModule& LevelEditor = FModuleManager::LoadModuleChecked<FLevelEditorModule>("LevelEditor");
-		MapChangedHandle = LevelEditor.OnMapChanged().AddRaw(this, &FConcertClientWorkspace::HandleMapChanged);
-	}
-
 	// Register Object Transaction events
 	if (GUnrealEd)
 	{
@@ -377,16 +370,6 @@ void FConcertClientWorkspace::UnbindSession()
 			EndPIEHandle.Reset();
 		}
 
-		// Unregister Map Change Events
-		if (MapChangedHandle.IsValid())
-		{
-			if (FLevelEditorModule* LevelEditor = FModuleManager::GetModulePtr<FLevelEditorModule>("LevelEditor"))
-			{
-				LevelEditor->OnMapChanged().Remove(MapChangedHandle);
-			}
-			MapChangedHandle.Reset();
-		}
-
 		// Unregister Object Transaction events
 		if (GUnrealEd && TransactionStateChangedHandle.IsValid())
 		{
@@ -624,20 +607,9 @@ void FConcertClientWorkspace::HandleAssetRenamed(const FAssetData& Data, const F
 
 void FConcertClientWorkspace::HandleAssetLoaded(UObject* InAsset)
 {
-	const FName LoadedPackageName = InAsset->GetOutermost()->GetFName();
-
-	// Skip world assets that are being loaded as the editor world
-	// We handle these via HandleMapChanged instead (which also calls this function, but the entry in UWorld::WorldTypePreLoadMap will have been removed by then)
-	if (const EWorldType::Type* WorldTypePtr = UWorld::WorldTypePreLoadMap.Find(LoadedPackageName))
-	{
-		if (*WorldTypePtr == EWorldType::Editor)
-		{
-			return;
-		}
-	}
-
-	if (TransactionManager.IsValid())
+	if (TransactionManager.IsValid() && bHasSyncedWorkspace)
 	{
+		const FName LoadedPackageName = InAsset->GetOutermost()->GetFName();
 		TransactionManager->ReplayTransactions(LoadedPackageName);
 	}
 }
@@ -696,14 +668,6 @@ void FConcertClientWorkspace::HandleEndPIE(const bool InIsSimulating)
 	}
 }
 
-void FConcertClientWorkspace::HandleMapChanged(UWorld* InWorld, EMapChangeType InMapChangeType)
-{
-	if (InMapChangeType == EMapChangeType::NewMap || InMapChangeType == EMapChangeType::LoadMap)
-	{
-		HandleAssetLoaded(InWorld);
-	}
-}
-
 void FConcertClientWorkspace::HandleTransactionStateChanged(const FTransactionContext& InTransactionContext, const ETransactionStateEventType InTransactionState)
 {
 	if (TransactionManager.IsValid())
@@ -743,6 +707,9 @@ void FConcertClientWorkspace::OnEndFrame()
 		}
 		TransactionManager->ReplayAllTransactions();
 
+		// We process all pending transactions we just replayed before finalizing the sync to prevent package being loaded as a result to trigger replaying transactions again
+		TransactionManager->ProcessPending();
+
 		// Finalize the sync
 		bHasSyncedWorkspace = true;
 		InitialSyncSlowTask.Reset();
diff --git a/Engine/Plugins/Developer/Concert/ConcertSync/ConcertSyncClient/Source/ConcertSyncClient/Private/ConcertClientWorkspace.h b/Engine/Plugins/Developer/Concert/ConcertSync/ConcertSyncClient/Source/ConcertSyncClient/Private/ConcertClientWorkspace.h
--- a/Engine/Plugins/Developer/Concert/ConcertSync/ConcertSyncClient/Source/ConcertSyncClient/Private/ConcertClientWorkspace.h
+++ b/Engine/Plugins/Developer/Concert/ConcertSync/ConcertSyncClient/Source/ConcertSyncClient/Private/ConcertClientWorkspace.h
@@ -89,9 +89,6 @@ private:
 	/** */
 	void HandleEndPIE(const bool InIsSimulating);
 
-	/** */
-	void HandleMapChanged(UWorld* InWorld, EMapChangeType InMapChangeType);
-
 	/** */
 	void HandleTransactionStateChanged(const FTransactionContext& InTransactionContext, const ETransactionStateEventType InTransactionState);
 
@@ -171,9 +168,6 @@ private:
 	/** */
 	FDelegateHandle EndPIEHandle;
 
-	/** */
-	FDelegateHandle MapChangedHandle;
-
 	/** */
 	FDelegateHandle TransactionStateChangedHandle;
 	
