Modified Multi-User presence event to indicate which type of VR device a specific user is using so that avatar presence can use the proper HMD mesh representation

diff --git a/Engine/Plugins/Developer/Concert/ConcertSync/ConcertSyncClient/Source/ConcertSyncClient/Private/ConcertClientDesktopPresenceActor.cpp b/Engine/Plugins/Developer/Concert/ConcertSync/ConcertSyncClient/Source/ConcertSyncClient/Private/ConcertClientDesktopPresenceActor.cpp
--- a/Engine/Plugins/Developer/Concert/ConcertSync/ConcertSyncClient/Source/ConcertSyncClient/Private/ConcertClientDesktopPresenceActor.cpp
+++ b/Engine/Plugins/Developer/Concert/ConcertSync/ConcertSyncClient/Source/ConcertSyncClient/Private/ConcertClientDesktopPresenceActor.cpp
@@ -54,9 +54,9 @@ void AConcertClientDesktopPresenceActor::SetPresenceColor(const FLinearColor& In
 	TextMID->SetVectorParameterValue(ColorParamName, InColor);
 }
 
-void AConcertClientDesktopPresenceActor::InitPresence(const class UConcertAssetContainer& InAssetContainer)
+void AConcertClientDesktopPresenceActor::InitPresence(const class UConcertAssetContainer& InAssetContainer, FName DeviceType)
 {
-	Super::InitPresence(InAssetContainer);
+	Super::InitPresence(InAssetContainer, DeviceType);
 
 	UMaterialInterface* PresenceMaterial = InAssetContainer.PresenceFadeMaterial;
 
diff --git a/Engine/Plugins/Developer/Concert/ConcertSync/ConcertSyncClient/Source/ConcertSyncClient/Private/ConcertClientDesktopPresenceActor.h b/Engine/Plugins/Developer/Concert/ConcertSync/ConcertSyncClient/Source/ConcertSyncClient/Private/ConcertClientDesktopPresenceActor.h
--- a/Engine/Plugins/Developer/Concert/ConcertSync/ConcertSyncClient/Source/ConcertSyncClient/Private/ConcertClientDesktopPresenceActor.h
+++ b/Engine/Plugins/Developer/Concert/ConcertSync/ConcertSyncClient/Source/ConcertSyncClient/Private/ConcertClientDesktopPresenceActor.h
@@ -27,7 +27,7 @@ public:
 	virtual void SetPresenceColor(const FLinearColor& InColor) override;
 
 	/** AConcertClientPresenceActor Interface */
-	virtual void InitPresence(const class UConcertAssetContainer& InAssetContainer) override;
+	virtual void InitPresence(const class UConcertAssetContainer& InAssetContainer, FName DeviceType) override;
 
 	/** Handle presence update events */
 	virtual void HandleEvent(const FStructOnScope& InEvent) override;
diff --git a/Engine/Plugins/Developer/Concert/ConcertSync/ConcertSyncClient/Source/ConcertSyncClient/Private/ConcertClientPresenceActor.cpp b/Engine/Plugins/Developer/Concert/ConcertSync/ConcertSyncClient/Source/ConcertSyncClient/Private/ConcertClientPresenceActor.cpp
--- a/Engine/Plugins/Developer/Concert/ConcertSync/ConcertSyncClient/Source/ConcertSyncClient/Private/ConcertClientPresenceActor.cpp
+++ b/Engine/Plugins/Developer/Concert/ConcertSync/ConcertSyncClient/Source/ConcertSyncClient/Private/ConcertClientPresenceActor.cpp
@@ -116,7 +116,7 @@ void AConcertClientPresenceActor::Tick(float DeltaSeconds)
 	}
 }
 
-void AConcertClientPresenceActor::InitPresence(const UConcertAssetContainer& InAssetContainer)
+void AConcertClientPresenceActor::InitPresence(const UConcertAssetContainer& InAssetContainer, FName /* DeviceType */)
 {
 	UStaticMesh* PresenceMesh = InAssetContainer.GenericDesktopMesh;
 
diff --git a/Engine/Plugins/Developer/Concert/ConcertSync/ConcertSyncClient/Source/ConcertSyncClient/Private/ConcertClientPresenceActor.h b/Engine/Plugins/Developer/Concert/ConcertSync/ConcertSyncClient/Source/ConcertSyncClient/Private/ConcertClientPresenceActor.h
--- a/Engine/Plugins/Developer/Concert/ConcertSync/ConcertSyncClient/Source/ConcertSyncClient/Private/ConcertClientPresenceActor.h
+++ b/Engine/Plugins/Developer/Concert/ConcertSync/ConcertSyncClient/Source/ConcertSyncClient/Private/ConcertClientPresenceActor.h
@@ -49,7 +49,7 @@ public:
 
 	virtual void HandleEvent(const FStructOnScope& InEvent);
 
-	virtual void InitPresence(const class UConcertAssetContainer& InAssetContainer);
+	virtual void InitPresence(const class UConcertAssetContainer& InAssetContainer, FName DeviceType);
 
 	virtual bool ShouldTickIfViewportsOnly() const override;
 
diff --git a/Engine/Plugins/Developer/Concert/ConcertSync/ConcertSyncClient/Source/ConcertSyncClient/Private/ConcertClientPresenceManager.cpp b/Engine/Plugins/Developer/Concert/ConcertSync/ConcertSyncClient/Source/ConcertSyncClient/Private/ConcertClientPresenceManager.cpp
--- a/Engine/Plugins/Developer/Concert/ConcertSync/ConcertSyncClient/Source/ConcertSyncClient/Private/ConcertClientPresenceManager.cpp
+++ b/Engine/Plugins/Developer/Concert/ConcertSync/ConcertSyncClient/Source/ConcertSyncClient/Private/ConcertClientPresenceManager.cpp
@@ -34,6 +34,7 @@
 #include "SLevelViewport.h"
 #include "ConcertAssetContainer.h"
 #include "IVREditorModule.h"
+#include "VREditorMode.h"
 #include "Framework/Application/SlateApplication.h"
 #endif
 
@@ -69,7 +70,7 @@ FConcertClientPresenceManager::FConcertClientPresenceManager(TSharedRef<IConcert
 	, CurrentAvatarMode(nullptr)
 	, AssetContainer(nullptr)
 	, bIsPresenceEnabled(true)
-	, bInVR(false)
+	, VRDeviceType(NAME_None)
 	, CurrentAvatarActorClass(nullptr)
 	, DesktopAvatarActorClass(nullptr)
 	, VRAvatarActorClass(nullptr)
@@ -305,7 +306,7 @@ void FConcertClientPresenceManager::SynchronizePresenceState()
 			Session->FindSessionClient(RemoteEndpointId, ClientSessionInfo);
 			if (!PresenceState.PresenceActor.IsValid())
 			{
-				PresenceState.PresenceActor = CreatePresenceActor(ClientSessionInfo.ClientInfo, PresenceState.bInVR);
+				PresenceState.PresenceActor = CreatePresenceActor(ClientSessionInfo.ClientInfo, PresenceState.VRDevice);
 			}
 
 			if (PresenceState.PresenceActor.IsValid())
@@ -364,9 +365,9 @@ bool FConcertClientPresenceManager::ShouldProcessPresenceEvent(const FConcertSes
 	return true;
 }
 
-AConcertClientPresenceActor* FConcertClientPresenceManager::CreatePresenceActor(const FConcertClientInfo& InClientInfo, bool bClientInVR)
+AConcertClientPresenceActor* FConcertClientPresenceManager::CreatePresenceActor(const FConcertClientInfo& InClientInfo, FName VRDevice)
 {
-	AConcertClientPresenceActor* PresenceActor = SpawnPresenceActor(InClientInfo, bClientInVR);
+	AConcertClientPresenceActor* PresenceActor = SpawnPresenceActor(InClientInfo, VRDevice);
 
 	if (PresenceActor)
 	{
@@ -377,7 +378,7 @@ AConcertClientPresenceActor* FConcertClientPresenceManager::CreatePresenceActor(
 	return PresenceActor;
 }
 
-AConcertClientPresenceActor* FConcertClientPresenceManager::SpawnPresenceActor(const FConcertClientInfo& InClientInfo, bool bClientInVR)
+AConcertClientPresenceActor* FConcertClientPresenceManager::SpawnPresenceActor(const FConcertClientInfo& InClientInfo, FName VRDevice)
 {
 	check(AssetContainer);
 
@@ -390,7 +391,7 @@ AConcertClientPresenceActor* FConcertClientPresenceManager::SpawnPresenceActor(c
 
 	// @todo this is potentially slow and hitchy as clients connect.  It might be better to preload all the presence actor types
 	UClass* PresenceActorClass = nullptr;	
-	if (bClientInVR)
+	if (!VRDevice.IsNone())
 	{
 		PresenceActorClass = LoadObject<UClass>(nullptr, *InClientInfo.VRAvatarActorClass);
 	}
@@ -401,7 +402,7 @@ AConcertClientPresenceActor* FConcertClientPresenceManager::SpawnPresenceActor(c
 
 	if (!PresenceActorClass)
 	{
-		UE_LOG(LogConcert, Warning, TEXT("Failed to load presence actor class '%s'. Presence will not be displayed"), bClientInVR ? *InClientInfo.VRAvatarActorClass : *InClientInfo.DesktopAvatarActorClass);
+		UE_LOG(LogConcert, Warning, TEXT("Failed to load presence actor class '%s'. Presence will not be displayed"), !VRDevice.IsNone() ? *InClientInfo.VRAvatarActorClass : *InClientInfo.DesktopAvatarActorClass);
 		return nullptr;
 	}
 
@@ -425,12 +426,12 @@ AConcertClientPresenceActor* FConcertClientPresenceManager::SpawnPresenceActor(c
 
 	if (!PresenceActor)
 	{
-		UE_LOG(LogConcert, Warning, TEXT("Failed to spawn presence actor of class '%s'. Presence will not be displayed"), bClientInVR ? *InClientInfo.VRAvatarActorClass : *InClientInfo.DesktopAvatarActorClass);
+		UE_LOG(LogConcert, Warning, TEXT("Failed to spawn presence actor of class '%s'. Presence will not be displayed"), !VRDevice.IsNone() ? *InClientInfo.VRAvatarActorClass : *InClientInfo.DesktopAvatarActorClass);
 		return nullptr;
 	}
 
 	// Setup the asset container.
-	PresenceActor->InitPresence(*AssetContainer);
+	PresenceActor->InitPresence(*AssetContainer, VRDevice);
 
 	return PresenceActor;
 }
@@ -514,24 +515,25 @@ void FConcertClientPresenceManager::OnSessionClientChanged(IConcertClientSession
 
 void FConcertClientPresenceManager::OnVREditingModeEnter()
 {
-	bInVR = true;
+	UVREditorMode* VRMode = IVREditorModule::Get().GetVRMode();
+	VRDeviceType = VRMode ? VRMode->GetHMDDeviceType() : FName();
 	UpdatePresenceMode();
 }
 
 void FConcertClientPresenceManager::OnVREditingModeExit()
 {
-	bInVR = false;
+	VRDeviceType = FName();
 	UpdatePresenceMode();
 }
 
 void FConcertClientPresenceManager::UpdatePresenceMode()
 {
-	if ((bInVR && CurrentAvatarActorClass != VRAvatarActorClass) ||
-		(!bInVR && CurrentAvatarActorClass != DesktopAvatarActorClass))
+	if ((!VRDeviceType.IsNone() && CurrentAvatarActorClass != VRAvatarActorClass) ||
+		(VRDeviceType.IsNone() && CurrentAvatarActorClass != DesktopAvatarActorClass))
 	{
 		// Mode will get recreated on next call to OnEndFrame
 		CurrentAvatarMode.Reset();
-		CurrentAvatarActorClass = bInVR ? VRAvatarActorClass : DesktopAvatarActorClass;
+		CurrentAvatarActorClass = !VRDeviceType.IsNone() ? VRAvatarActorClass : DesktopAvatarActorClass;
 		SendPresenceInVREvent();
 	}
 }
@@ -539,7 +541,7 @@ void FConcertClientPresenceManager::UpdatePresenceMode()
 void FConcertClientPresenceManager::SendPresenceInVREvent(const FGuid* InEndpointId)
 {
 	FConcertClientPresenceInVREvent Event;
-	Event.bInVR = bInVR;
+	Event.VRDevice = VRDeviceType;
 
 	if (InEndpointId)
 	{
@@ -553,13 +555,13 @@ void FConcertClientPresenceManager::SendPresenceInVREvent(const FGuid* InEndpoin
 
 void FConcertClientPresenceManager::HandleConcertClientPresenceInVREvent(const FConcertSessionContext& InSessionContext, const FConcertClientPresenceInVREvent& InEvent)
 {
-	UpdatePresenceAvatar(InSessionContext.SourceEndpointId, InEvent.bInVR);
+	UpdatePresenceAvatar(InSessionContext.SourceEndpointId, InEvent.VRDevice);
 }
 
-void FConcertClientPresenceManager::UpdatePresenceAvatar(const FGuid& InEndpointId, bool bIsInVR)
+void FConcertClientPresenceManager::UpdatePresenceAvatar(const FGuid& InEndpointId, FName VRDevice)
 {
 	FConcertClientPresenceState& PresenceState = EnsurePresenceState(InEndpointId);
-	PresenceState.bInVR = bIsInVR;
+	PresenceState.VRDevice = VRDevice;
 
 	if (PresenceState.PresenceActor.IsValid())
 	{
@@ -725,7 +727,7 @@ FReply FConcertClientPresenceManager::OnJumpToPresenceClicked(FGuid InEndpointId
 		FRotator OtherClientRotation(OtherClientState->Orientation.Rotator());
 
 		// Disregard pitch and roll when teleporting to a VR presence.
-		if (bInVR)
+		if (!VRDeviceType.IsNone())
 		{
 			OtherClientRotation.Pitch = 0.0f;
 			OtherClientRotation.Roll = 0.0f;
diff --git a/Engine/Plugins/Developer/Concert/ConcertSync/ConcertSyncClient/Source/ConcertSyncClient/Private/ConcertClientPresenceManager.h b/Engine/Plugins/Developer/Concert/ConcertSync/ConcertSyncClient/Source/ConcertSyncClient/Private/ConcertClientPresenceManager.h
--- a/Engine/Plugins/Developer/Concert/ConcertSync/ConcertSyncClient/Source/ConcertSyncClient/Private/ConcertClientPresenceManager.h
+++ b/Engine/Plugins/Developer/Concert/ConcertSync/ConcertSyncClient/Source/ConcertSyncClient/Private/ConcertClientPresenceManager.h
@@ -38,9 +38,10 @@ struct FConcertClientPresenceState
 {
 	FConcertClientPresenceState()
 		: bIsConnected(true)
+		, bVisible(true)
 		, bInPIE(false)
-		, bInVR(false)
-		, bVisible(true) {}
+		, VRDevice(NAME_None)
+	{}
 
 	/** State map */
 	TMap<UScriptStruct*, FConcertClientPresenceStateEntry> EventStateMap;
@@ -51,14 +52,14 @@ struct FConcertClientPresenceState
 	/** Whether client is connected */
 	bool bIsConnected;
 
+	/** Whether client is visible */
+	bool bVisible;
+
 	/** Whether client is in PIE */
 	bool bInPIE;
 
-	/** Whether client is in VR editing mode */
-	bool bInVR;
-
-	/** Whether client is visible */
-	bool bVisible;
+	/** Whether client is using a VRDevice */
+	FName VRDevice;
 
 	/** Presence actor */
 	TWeakObjectPtr<AConcertClientPresenceActor> PresenceActor;
@@ -162,10 +163,10 @@ private:
 	bool ShouldProcessPresenceEvent(const FConcertSessionContext& InSessionContext, const UStruct* InEventType, const FConcertClientPresenceEventBase& InEvent) const;
 
 	/** Create a new presence actor */
-	AConcertClientPresenceActor* CreatePresenceActor(const FConcertClientInfo& InClientInfo, bool bClientInVR);
+	AConcertClientPresenceActor* CreatePresenceActor(const FConcertClientInfo& InClientInfo, FName VRDevice);
 
 	/** Spawn a presence actor */
-	AConcertClientPresenceActor* SpawnPresenceActor(const FConcertClientInfo& InClientInfo, bool bClientInVR);
+	AConcertClientPresenceActor* SpawnPresenceActor(const FConcertClientInfo& InClientInfo, FName VRDevice);
 
 	/** Clear presence */
 	void ClearPresenceActor(const FGuid& InEndpointId);
@@ -198,7 +199,7 @@ private:
 	void HandleConcertClientPresenceInVREvent(const FConcertSessionContext& InSessionContext, const FConcertClientPresenceInVREvent& InEvent);
 
 	/** Updates presence avatar for remote client by invalidating current presence actor */
-	void UpdatePresenceAvatar(const FGuid& InEndpointId, bool bIsInVR);
+	void UpdatePresenceAvatar(const FGuid& InEndpointId, FName VRDevice);
 
 	/** Set presence PIE state */
 	void SetPresenceInPIE(const FGuid& InEndpointId, bool bInPIE);
@@ -263,8 +264,8 @@ private:
 	/** True if presence is currently enabled and should be shown (unless hidden by other settings) */
 	bool bIsPresenceEnabled;
 
-	/** True if in VR */
-	bool bInVR;
+	/** NAME_None if not in VR */
+	FName VRDeviceType;
 
 	/** Avatar actor class */
 	UClass* CurrentAvatarActorClass;
diff --git a/Engine/Plugins/Developer/Concert/ConcertSync/ConcertSyncClient/Source/ConcertSyncClient/Private/ConcertClientVRPresenceActor.cpp b/Engine/Plugins/Developer/Concert/ConcertSync/ConcertSyncClient/Source/ConcertSyncClient/Private/ConcertClientVRPresenceActor.cpp
--- a/Engine/Plugins/Developer/Concert/ConcertSync/ConcertSyncClient/Source/ConcertSyncClient/Private/ConcertClientVRPresenceActor.cpp
+++ b/Engine/Plugins/Developer/Concert/ConcertSync/ConcertSyncClient/Source/ConcertSyncClient/Private/ConcertClientVRPresenceActor.cpp
@@ -149,6 +149,7 @@ void AConcertClientVRPresenceActor::Tick(float DeltaSeconds)
 
 		const FTransform RightControllerTransform(RightControllerOrientation, RightControllerPosition);
 		RightControllerMeshComponent->SetWorldTransform(RightControllerTransform);
+		RightControllerMeshComponent->SetRelativeScale3D(FVector(1.0f, -1.0f, 1.0f));
 	}
 
 	// Calculate laser
@@ -167,12 +168,12 @@ void AConcertClientVRPresenceActor::Tick(float DeltaSeconds)
 	}
 }
 
-void AConcertClientVRPresenceActor::InitPresence(const class UConcertAssetContainer& InAssetContainer)
+void AConcertClientVRPresenceActor::InitPresence(const class UConcertAssetContainer& InAssetContainer, FName DeviceType)
 {
-	Super::InitPresence(InAssetContainer);
+	Super::InitPresence(InAssetContainer, DeviceType);
 
 	// To do, send data about these through the event.
-	UStaticMesh* ControllerMesh = InAssetContainer.VivePreControllerMesh;
+	UStaticMesh* ControllerMesh = DeviceType == FName(TEXT("OculusHMD")) ? InAssetContainer.OculusControllerMesh : InAssetContainer.VivePreControllerMesh;
 
 	LeftControllerMeshComponent->SetStaticMesh(ControllerMesh);
 	LeftControllerMeshComponent->SetMobility(EComponentMobility::Movable);
diff --git a/Engine/Plugins/Developer/Concert/ConcertSync/ConcertSyncClient/Source/ConcertSyncClient/Private/ConcertClientVRPresenceActor.h b/Engine/Plugins/Developer/Concert/ConcertSync/ConcertSyncClient/Source/ConcertSyncClient/Private/ConcertClientVRPresenceActor.h
--- a/Engine/Plugins/Developer/Concert/ConcertSync/ConcertSyncClient/Source/ConcertSyncClient/Private/ConcertClientVRPresenceActor.h
+++ b/Engine/Plugins/Developer/Concert/ConcertSync/ConcertSyncClient/Source/ConcertSyncClient/Private/ConcertClientVRPresenceActor.h
@@ -19,7 +19,7 @@ class AConcertClientVRPresenceActor : public AConcertClientPresenceActor
 public:
 	virtual void HandleEvent(const FStructOnScope& InEvent) override;
 
-	virtual void InitPresence(const class UConcertAssetContainer& InAssetContainer) override;
+	virtual void InitPresence(const class UConcertAssetContainer& InAssetContainer, FName DeviceType) override;
 
 	virtual void SetPresenceColor(const FLinearColor& InColor) override;
 
diff --git a/Engine/Plugins/Developer/Concert/ConcertSync/ConcertSyncCore/Source/ConcertSyncCore/Public/ConcertPresenceEvents.h b/Engine/Plugins/Developer/Concert/ConcertSync/ConcertSyncCore/Source/ConcertSyncCore/Public/ConcertPresenceEvents.h
--- a/Engine/Plugins/Developer/Concert/ConcertSync/ConcertSyncCore/Source/ConcertSyncCore/Public/ConcertPresenceEvents.h
+++ b/Engine/Plugins/Developer/Concert/ConcertSync/ConcertSyncCore/Source/ConcertSyncCore/Public/ConcertPresenceEvents.h
@@ -32,7 +32,7 @@ struct FConcertClientPresenceInVREvent
 	GENERATED_BODY()
 
 	UPROPERTY()
-	bool bInVR;
+	FName VRDevice;
 };
 
 USTRUCT()
