Adjusted avatar presence hand position to match oculus vr handle

diff --git a/Engine/Plugins/Developer/Concert/ConcertSync/ConcertSyncClient/Source/ConcertSyncClient/Private/ConcertClientPresenceActor.cpp b/Engine/Plugins/Developer/Concert/ConcertSync/ConcertSyncClient/Source/ConcertSyncClient/Private/ConcertClientPresenceActor.cpp
--- a/Engine/Plugins/Developer/Concert/ConcertSync/ConcertSyncClient/Source/ConcertSyncClient/Private/ConcertClientPresenceActor.cpp
+++ b/Engine/Plugins/Developer/Concert/ConcertSync/ConcertSyncClient/Source/ConcertSyncClient/Private/ConcertClientPresenceActor.cpp
@@ -116,8 +116,9 @@ void AConcertClientPresenceActor::Tick(float DeltaSeconds)
 	}
 }
 
-void AConcertClientPresenceActor::InitPresence(const UConcertAssetContainer& InAssetContainer, FName /* DeviceType */)
+void AConcertClientPresenceActor::InitPresence(const UConcertAssetContainer& InAssetContainer, FName DeviceType)
 {
+	PresenceDeviceType = DeviceType;
 	UStaticMesh* PresenceMesh = InAssetContainer.GenericDesktopMesh;
 
 	if (PresenceMeshComponent->GetStaticMesh() == nullptr)
diff --git a/Engine/Plugins/Developer/Concert/ConcertSync/ConcertSyncClient/Source/ConcertSyncClient/Private/ConcertClientPresenceActor.h b/Engine/Plugins/Developer/Concert/ConcertSync/ConcertSyncClient/Source/ConcertSyncClient/Private/ConcertClientPresenceActor.h
--- a/Engine/Plugins/Developer/Concert/ConcertSync/ConcertSyncClient/Source/ConcertSyncClient/Private/ConcertClientPresenceActor.h
+++ b/Engine/Plugins/Developer/Concert/ConcertSync/ConcertSyncClient/Source/ConcertSyncClient/Private/ConcertClientPresenceActor.h
@@ -56,6 +56,10 @@ public:
 	virtual void Tick(float DeltaSeconds) override;
 
 protected:
+	/* The device type that this presence represent (i.e Oculus, Vive, Desktop) */
+	UPROPERTY(BlueprintReadOnly, Category = "Rendering", meta = (AllowPrivateAccess = "true"))
+	FName PresenceDeviceType;
+
 	/** The camera mesh component to show visually where the camera is placed */
 	UPROPERTY(VisibleAnywhere, BlueprintReadWrite, Category = "Rendering", meta = (AllowPrivateAccess = "true"))
 	class UStaticMeshComponent* PresenceMeshComponent;
diff --git a/Engine/Plugins/Developer/Concert/ConcertSync/ConcertSyncClient/Source/ConcertSyncClient/Private/ConcertClientPresenceManager.cpp b/Engine/Plugins/Developer/Concert/ConcertSync/ConcertSyncClient/Source/ConcertSyncClient/Private/ConcertClientPresenceManager.cpp
--- a/Engine/Plugins/Developer/Concert/ConcertSync/ConcertSyncClient/Source/ConcertSyncClient/Private/ConcertClientPresenceManager.cpp
+++ b/Engine/Plugins/Developer/Concert/ConcertSync/ConcertSyncClient/Source/ConcertSyncClient/Private/ConcertClientPresenceManager.cpp
@@ -414,6 +414,7 @@ AConcertClientPresenceActor* FConcertClientPresenceManager::SpawnPresenceActor(c
 		ActorSpawnParameters.Name = MakeUniqueObjectName(World, PresenceActorClass, PresenceActorClass->GetFName()); // @todo how should spawned actors be named?
 		ActorSpawnParameters.SpawnCollisionHandlingOverride = ESpawnActorCollisionHandlingMethod::AlwaysSpawn;
 		ActorSpawnParameters.ObjectFlags = EObjectFlags::RF_DuplicateTransient;
+		ActorSpawnParameters.bDeferConstruction = true;
 
 		PresenceActor = World->SpawnActor<AConcertClientPresenceActor>(PresenceActorClass, ActorSpawnParameters);
 
@@ -432,6 +433,10 @@ AConcertClientPresenceActor* FConcertClientPresenceManager::SpawnPresenceActor(c
 
 	// Setup the asset container.
 	PresenceActor->InitPresence(*AssetContainer, VRDevice);
+	{
+		FEditorScriptExecutionGuard UCSGuard;
+		PresenceActor->FinishSpawning(FTransform(), true);
+	}
 
 	return PresenceActor;
 }
diff --git a/Engine/Plugins/Developer/Concert/ConcertSync/ConcertSyncClient/Source/ConcertSyncClient/Private/ConcertClientVRPresenceActor.cpp b/Engine/Plugins/Developer/Concert/ConcertSync/ConcertSyncClient/Source/ConcertSyncClient/Private/ConcertClientVRPresenceActor.cpp
--- a/Engine/Plugins/Developer/Concert/ConcertSync/ConcertSyncClient/Source/ConcertSyncClient/Private/ConcertClientVRPresenceActor.cpp
+++ b/Engine/Plugins/Developer/Concert/ConcertSync/ConcertSyncClient/Source/ConcertSyncClient/Private/ConcertClientVRPresenceActor.cpp
@@ -173,7 +173,7 @@ void AConcertClientVRPresenceActor::InitPresence(const class UConcertAssetContai
 	Super::InitPresence(InAssetContainer, DeviceType);
 
 	// To do, send data about these through the event.
-	UStaticMesh* ControllerMesh = DeviceType == FName(TEXT("OculusHMD")) ? InAssetContainer.OculusControllerMesh : InAssetContainer.VivePreControllerMesh;
+	UStaticMesh* ControllerMesh = PresenceDeviceType == FName(TEXT("OculusHMD")) ? InAssetContainer.OculusControllerMesh : InAssetContainer.VivePreControllerMesh;
 
 	LeftControllerMeshComponent->SetStaticMesh(ControllerMesh);
 	LeftControllerMeshComponent->SetMobility(EComponentMobility::Movable);
