- Replaced the current editor toolbar icons to use more representative icons. It has 2 new icons to represents states of the button that were not represented before (cancel and browse).

diff --git a/Engine/Plugins/Developer/Concert/ConcertSync/ConcertSyncClient/Source/ConcertSyncUI/Private/ConcertFrontendStyle.cpp b/Engine/Plugins/Developer/Concert/ConcertSync/ConcertSyncClient/Source/ConcertSyncUI/Private/ConcertFrontendStyle.cpp
--- a/Engine/Plugins/Developer/Concert/ConcertSync/ConcertSyncClient/Source/ConcertSyncUI/Private/ConcertFrontendStyle.cpp
+++ b/Engine/Plugins/Developer/Concert/ConcertSync/ConcertSyncClient/Source/ConcertSyncUI/Private/ConcertFrontendStyle.cpp
@@ -82,10 +82,6 @@ void FConcertFrontendStyle::Initialize()
 	StyleSet->Set("Concert.HidePresence",       new IMAGE_PLUGIN_BRUSH("Icons/icon_PresenceEyeOff_32x",     Icon16x16, IconColorAndOpacity));
 	StyleSet->Set("Concert.ShowPresence",       new IMAGE_PLUGIN_BRUSH("Icons/icon_PresenceEyeOn_32x",      Icon16x16, IconColorAndOpacity));
 
-	// 20x20 -> Editor toolbar small icons.
-	StyleSet->Set("Concert.Online.Small", new IMAGE_PLUGIN_BRUSH("Icons/icon_ConcertOnline_40x", Icon20x20));
-	StyleSet->Set("Concert.Offline.Small", new IMAGE_PLUGIN_BRUSH("Icons/icon_ConcertOffline_40x", Icon20x20));
-
 	// 24x24/48x48 -> For sequencer toolbar.
 	StyleSet->Set("Concert.Sequencer.SyncTimeline",       new IMAGE_PLUGIN_BRUSH("Icons/icon_SequencerSyncTimeline_48x", Icon48x48, IconColorAndOpacity)); // Enable/disable playback and time scrubbing from a remote client.
 	StyleSet->Set("Concert.Sequencer.SyncTimeline.Small", new IMAGE_PLUGIN_BRUSH("Icons/icon_SequencerSyncTimeline_48x", Icon24x24, IconColorAndOpacity));
@@ -93,8 +89,16 @@ void FConcertFrontendStyle::Initialize()
 	StyleSet->Set("Concert.Sequencer.SyncSequence.Small", new IMAGE_PLUGIN_BRUSH("Icons/icon_SequencerSyncSequence_48x", Icon24x24, IconColorAndOpacity));
 
 	// 40x40 -> Editor toolbar large icons.
-	StyleSet->Set("Concert.Online", new IMAGE_PLUGIN_BRUSH("Icons/icon_ConcertOnline_40x", Icon40x40));
-	StyleSet->Set("Concert.Offline", new IMAGE_PLUGIN_BRUSH("Icons/icon_ConcertOffline_40x", Icon40x40));
+	StyleSet->Set("Concert.Browse", new IMAGE_PLUGIN_BRUSH("Icons/icon_MultiUserMenuBrowse_40x", Icon40x40));
+	StyleSet->Set("Concert.Join",   new IMAGE_PLUGIN_BRUSH("Icons/icon_MultiUserMenuJoin_40x",   Icon40x40));
+	StyleSet->Set("Concert.Leave",  new IMAGE_PLUGIN_BRUSH("Icons/icon_MultiUserMenuLeave_40x",  Icon40x40));
+	StyleSet->Set("Concert.Cancel", new IMAGE_PLUGIN_BRUSH("Icons/icon_MultiUserMenuCancel_40x", Icon40x40));
+
+	// 20x20 -> Editor toolbar small icons.
+	StyleSet->Set("Concert.Browse.Small", new IMAGE_PLUGIN_BRUSH("Icons/icon_MultiUserMenuBrowse_40x", Icon20x20));
+	StyleSet->Set("Concert.Leave.Small",  new IMAGE_PLUGIN_BRUSH("Icons/icon_MultiUserMenuLeave_40x",  Icon20x20));
+	StyleSet->Set("Concert.Join.Small",   new IMAGE_PLUGIN_BRUSH("Icons/icon_MultiUserMenuJoin_40x",   Icon20x20));
+	StyleSet->Set("Concert.Cancel.Small", new IMAGE_PLUGIN_BRUSH("Icons/icon_MultiUserMenuCancel_40x", Icon20x20));
 
 	// Activity Text
 	{
diff --git a/Engine/Plugins/Developer/Concert/MultiUserClient/Source/MultiUserClient/Private/MultiUserClientModule.cpp b/Engine/Plugins/Developer/Concert/MultiUserClient/Source/MultiUserClient/Private/MultiUserClientModule.cpp
--- a/Engine/Plugins/Developer/Concert/MultiUserClient/Source/MultiUserClient/Private/MultiUserClientModule.cpp
+++ b/Engine/Plugins/Developer/Concert/MultiUserClient/Source/MultiUserClient/Private/MultiUserClientModule.cpp
@@ -729,9 +729,18 @@ private:
 		IConcertClientRef ConcertClient = MultiUserClient->GetConcertClient();
 		if (ConcertClient->GetSessionConnectionStatus() == EConcertConnectionStatus::Connected)
 		{
-			return FSlateIcon(FConcertFrontendStyle::GetStyleSetName(), "Concert.Online", "Concert.Online.Small");
+			return FSlateIcon(FConcertFrontendStyle::GetStyleSetName(), "Concert.Leave", "Concert.Leave.Small");
 		}
-		return FSlateIcon(FConcertFrontendStyle::GetStyleSetName(), "Concert.Offline", "Concert.Offline.Small");
+		else if (ConcertClient->IsAutoConnecting())
+		{
+			return FSlateIcon(FConcertFrontendStyle::GetStyleSetName(), "Concert.Cancel", "Concert.Cancel.Small");
+		}
+		else if (ConcertClient->CanAutoConnect())
+		{
+			return FSlateIcon(FConcertFrontendStyle::GetStyleSetName(), "Concert.Join", "Concert.Join.Small");
+		}
+		
+		return FSlateIcon(FConcertFrontendStyle::GetStyleSetName(), "Concert.Browse", "Concert.Browse.Small");
 	}
 
 	/** Return the toolbar icon tooltip. */
