What's happening (with respect to the bug report):

- When the BP editor was opened, it created the BlueprintActionDatabase, which reference the BP for its actions.
- When a Multi-User session is left, the BP created in the session is marked to be purged by Concert.
- In Concert, Purge + HotReload functions are executed back to back. In the purge function, the BP is marked to be unloaded. In the hot reload function, the map containing the BP is unloaded, but the BP is still referenced by the BlueprintActionDatabase and the BP doesn't get garbage collected.
- In the same tick, a little bit later, the directory watcher is ticked and find that .uasset corresponding to the BP was removed. The BlueprintActionDatabase gets notified and it removes its references on the BP, making it unreferenced and ready to be GC.
- The BP package remains in memory until the next GC (and they were no other GC in this case).
- On the next join session, keeping the same map, the hot reload function looks for the BP package and finds it before reloading the map. This is the BP package that we have purged, but not GC yet.
- Then, the hot reloading reloads the map, which triggers the GC. Since the BP package was ripe to be GC, it gets GC and the BP package we just found before gets dangling.

The solution:

- Force the BlueprintActionDatabase to release the references on the BP during the Purge so that the BP gets GC as soon as no map refers it anymore.

diff --git a/Engine/Plugins/Developer/Concert/ConcertSync/ConcertSyncClient/Source/ConcertSyncClient/Private/ConcertSyncClientUtil.cpp b/Engine/Plugins/Developer/Concert/ConcertSync/ConcertSyncClient/Source/ConcertSyncClient/Private/ConcertSyncClientUtil.cpp
--- a/Engine/Plugins/Developer/Concert/ConcertSync/ConcertSyncClient/Source/ConcertSyncClient/Private/ConcertSyncClientUtil.cpp
+++ b/Engine/Plugins/Developer/Concert/ConcertSync/ConcertSyncClient/Source/ConcertSyncClient/Private/ConcertSyncClientUtil.cpp
@@ -351,52 +351,54 @@ void PurgePackages(TArrayView<const FName> InPackageNames)
 		return;
 	}
 
+#if WITH_EDITOR
 	TArray<UObject*> ObjectsToPurge;
-	auto MakeObjectPurgeable = [&ObjectsToPurge](UObject* InObject)
+	auto CollectObjectToPurge = [&ObjectsToPurge](UObject* InObject)
 	{
-#if WITH_EDITOR
 		if (InObject->IsAsset())
 		{
 			FAssetEditorManager::Get().CloseAllEditorsForAsset(InObject);
 		}
-#endif
-		if (InObject->IsRooted())
-		{
-			InObject->RemoveFromRoot();
-		}
-		InObject->ClearFlags(RF_Public | RF_Standalone);
 		ObjectsToPurge.Add(InObject);
 	};
 
-	auto MakePackagePurgeable = [MakeObjectPurgeable](UPackage* InPackage)
-	{
-		// Prevent any message from the editor saying a package is not saved or doesn't exist on disk.
-		InPackage->SetDirtyFlag(false);
-
-		MakeObjectPurgeable(InPackage);
-		ForEachObjectWithOuter(InPackage, [MakeObjectPurgeable](UObject* InObject)
-		{
-			MakeObjectPurgeable(InObject);
-		});
-	};
-
-#if WITH_EDITOR
-	// Get the current edited map package to check if we are purging it
+	// Get the current edited map package to check if its going to be purged.
 	bool bEditedMapPurged = false;
 	UWorld* EditedWorld = GEditor->GetEditorWorldContext().World();
 	UPackage* EditedMapPackage = EditedWorld ? EditedWorld->GetOutermost() : nullptr;
-	
-	// Clean-up any in-memory packages that should be purged and check if we are purging the current map
+
+	// Collect any in-memory packages that should be purged and check if we are including the current map in the purge.
 	for (const FName PackageName : InPackageNames)
 	{
 		UPackage* ExistingPackage = FindPackage(nullptr, *PackageName.ToString());
 		if (ExistingPackage)
 		{
-			MakePackagePurgeable(ExistingPackage);
+			// Prevent any message from the editor saying a package is not saved or doesn't exist on disk.
+			ExistingPackage->SetDirtyFlag(false);
+
+			CollectObjectToPurge(ExistingPackage);
+			ForEachObjectWithOuter(ExistingPackage, [&CollectObjectToPurge](UObject* InObject)
+			{
+				CollectObjectToPurge(InObject);
+			});
+
 			bEditedMapPurged |= EditedMapPackage == ExistingPackage;
 		}
 	}
 
+	// Broadcast the eminent objects destruction (ex. tell BlueprintActionDatabase to release its reference(s) on Blueprint(s) right now)
+	FEditorDelegates::OnAssetsPreDelete.Broadcast(ObjectsToPurge);
+
+	// Mark objects as purgeable.
+	for (UObject* Object : ObjectsToPurge)
+	{
+		if (Object->IsRooted())
+		{
+			Object->RemoveFromRoot();
+		}
+		Object->ClearFlags(RF_Public | RF_Standalone);
+	}
+
 	// TODO: Revisit force replacing reference, current implementation is too aggressive and causes instability
 	// If we have any object that were made purgeable, null out their references so we can garbage collect
 	//if (ObjectsToPurge.Num() > 0)
@@ -412,11 +414,11 @@ void PurgePackages(TArrayView<const FName> InPackageNames)
 		FString StartupMapPackage = GetDefault<UGameMapsSettings>()->EditorStartupMap.GetLongPackageName();
 		if (FPackageName::DoesPackageExist(StartupMapPackage))
 		{
-			UEditorLoadingAndSavingUtils::NewMapFromTemplate(StartupMapPackage, /*bSaveExistingMap*/false);
+			UEditorLoadingAndSavingUtils::NewMapFromTemplate(StartupMapPackage, /*bSaveExistingMap*/false); // Expected to run GC internally.
 		}
 		else
 		{
-			UEditorLoadingAndSavingUtils::NewBlankMap(/*bSaveExistingMap*/false);
+			UEditorLoadingAndSavingUtils::NewBlankMap(/*bSaveExistingMap*/false); // Expected to run GC internally.
 		}
 	}
 	// if we have object to purge but the map isn't one of them collect garbage (if we purged the map it has already been done)
