- Ignored 'save' package events fired by the cooker by looking at 'GIsCookerLoadingPackage' global variable.

diff --git a/Engine/Plugins/Developer/Concert/ConcertSync/ConcertSyncClient/Source/ConcertSyncClient/Private/ConcertClientPackageBridge.cpp b/Engine/Plugins/Developer/Concert/ConcertSync/ConcertSyncClient/Source/ConcertSyncClient/Private/ConcertClientPackageBridge.cpp
--- a/Engine/Plugins/Developer/Concert/ConcertSync/ConcertSyncClient/Source/ConcertSyncClient/Private/ConcertClientPackageBridge.cpp
+++ b/Engine/Plugins/Developer/Concert/ConcertSync/ConcertSyncClient/Source/ConcertSyncClient/Private/ConcertClientPackageBridge.cpp
@@ -109,6 +109,13 @@ bool& FConcertClientPackageBridge::GetIgnoreLocalDiscardRef()
 
 void FConcertClientPackageBridge::HandlePackagePreSave(UPackage* Package)
 {
+	// Ignore package operations fired by the cooker (cook on the fly).
+	if (GIsCookerLoadingPackage)
+	{
+		check(IsInGameThread()); // We expect the cooker to call us on the game thread otherwise, we can have concurrency issues.
+		return;
+	}
+
 	// Ignore unwanted saves
 	if (bIgnoreLocalSave)
 	{
@@ -144,6 +151,13 @@ void FConcertClientPackageBridge::HandlePackageSaved(const FString& PackageFilen
 {
 	UPackage* Package = CastChecked<UPackage>(Outer);
 
+	// Ignore package operations fired by the cooker (cook on the fly).
+	if (GIsCookerLoadingPackage)
+	{
+		check(IsInGameThread()); // We expect the cooker to call us on the game thread otherwise, we can have concurrency issues.
+		return;
+	}
+
 	// Ignore unwanted saves
 	if (bIgnoreLocalSave)
 	{
