Fixed crash when calling PostTransacted on untracked matinee transaction
Fixed OngoingTransactions from never being removed from list

diff --git a/Engine/Plugins/Developer/Concert/ConcertSync/ConcertSyncClient/Source/ConcertSyncClient/Private/ConcertClientTransactionBridge.cpp b/Engine/Plugins/Developer/Concert/ConcertSync/ConcertSyncClient/Source/ConcertSyncClient/Private/ConcertClientTransactionBridge.cpp
--- a/Engine/Plugins/Developer/Concert/ConcertSync/ConcertSyncClient/Source/ConcertSyncClient/Private/ConcertClientTransactionBridge.cpp
+++ b/Engine/Plugins/Developer/Concert/ConcertSync/ConcertSyncClient/Source/ConcertSyncClient/Private/ConcertClientTransactionBridge.cpp
@@ -522,7 +522,8 @@ void FConcertClientTransactionBridge::HandleObjectTransacted(UObject* InObject,
 
 	UPackage* ChangedPackage = InObject->GetOutermost();
 	ConcertClientTransactionBridgeUtil::ETransactionFilterResult FilterResult = ConcertClientTransactionBridgeUtil::ApplyTransactionFilters(InObject, ChangedPackage);
-	
+	FOngoingTransaction* TrackedTransaction = OngoingTransactions.Find(InTransactionEvent.GetOperationId());
+
 	// TODO: This needs to send both editor-only and non-editor-only payload data to the server, which will forward only the correct part to cooked and non-cooked clients
 	bool bIncludeEditorOnlyProperties = true;
 
@@ -543,7 +544,8 @@ void FConcertClientTransactionBridge::HandleObjectTransacted(UObject* InObject,
 		}
 
 		UE_LOG(LogConcert, VeryVerbose,
-			TEXT("Transaction %s (%s, %s):%s %s:%s (%s property changes, %s object changes)"), 
+			TEXT("%s Transaction %s (%s, %s):%s %s:%s (%s property changes, %s object changes)"), 
+			(TrackedTransaction != nullptr ? TEXT("Tracked") : TEXT("Untracked")),
 			*InTransactionEvent.GetTransactionId().ToString(),
 			*InTransactionEvent.GetOperationId().ToString(),
 			ObjectEventString,
@@ -555,8 +557,13 @@ void FConcertClientTransactionBridge::HandleObjectTransacted(UObject* InObject,
 			);
 	}
 
+	if (TrackedTransaction == nullptr)
+	{
+		return;
+	}
+
 	const FConcertObjectId ObjectId = FConcertObjectId(*InObject->GetClass()->GetPathName(), InTransactionEvent.GetOriginalObjectOuterPathName(), InTransactionEvent.GetOriginalObjectName(), InObject->GetFlags());
-	FOngoingTransaction& OngoingTransaction = OngoingTransactions.FindChecked(InTransactionEvent.GetOperationId());
+	FOngoingTransaction& OngoingTransaction = *TrackedTransaction;
 
 	// If the object is excluded or exclude the whole transaction add it to the excluded list
 	if (FilterResult != ConcertClientTransactionBridgeUtil::ETransactionFilterResult::IncludeObject)
@@ -730,7 +737,7 @@ void FConcertClientTransactionBridge::OnEndFrame()
 		{
 			OnLocalTransactionFinalizedDelegate.Broadcast(OngoingTransactionPtr->CommonData, OngoingTransactionPtr->FinalizedData);
 
-			OngoingTransactions.Remove(OngoingTransactionPtr->CommonData.TransactionId);
+			OngoingTransactions.Remove(OngoingTransactionPtr->CommonData.OperationId);
 			OngoingTransactionsOrderIter.RemoveCurrent();
 			continue;
 		}
