Enabled Concert PresenceManager when in -game mode.
This allows -game client to see Multi-User presences if the Concert.EnablePresenceInGame console variable is true.
Merged from CL 7294055, CL 7294953, CL 7295479

diff --git a/Engine/Plugins/Developer/Concert/ConcertSync/ConcertSyncClient/Source/ConcertSyncClient/Private/ConcertClientPresenceActor.cpp b/Engine/Plugins/Developer/Concert/ConcertSync/ConcertSyncClient/Source/ConcertSyncClient/Private/ConcertClientPresenceActor.cpp
--- a/Engine/Plugins/Developer/Concert/ConcertSync/ConcertSyncClient/Source/ConcertSyncClient/Private/ConcertClientPresenceActor.cpp
+++ b/Engine/Plugins/Developer/Concert/ConcertSync/ConcertSyncClient/Source/ConcertSyncClient/Private/ConcertClientPresenceActor.cpp
@@ -11,15 +11,22 @@
 #include "ConcertClientPresenceManager.h"
 #include "ConcertClientMovement.h"
 #include "ConcertLogGlobal.h"
+#include "HAL/IConsoleManager.h"
 
 #define LOCTEXT_NAMESPACE "ConcertClientPresenceActor"
 
+static TAutoConsoleVariable<int32> CVarEnablePresenceInGame(TEXT("Concert.EnablePresenceInGame"), 0, TEXT("Enable Concert Presence in Game"));
+
+
 //////////////////////////////////////////////////////////////////////////
 // AConcertClientPresenceActor
 
 AConcertClientPresenceActor::AConcertClientPresenceActor(const FObjectInitializer& ObjectInitializer)
 	: Super(ObjectInitializer)
 {
+	// Initialize if this actor will be editor only based the `EnablePresenceInGame` console variable.
+	bIsEditorOnly = CVarEnablePresenceInGame.GetValueOnAnyThread() == 0;
+
 	// Set root component 
 	{
 		USceneComponent* SceneRootComponent = CreateDefaultSubobject<USceneComponent>(TEXT("Root"));
@@ -51,6 +58,11 @@ bool AConcertClientPresenceActor::ShouldTickIfViewportsOnly() const
 	return true;
 }
 
+bool AConcertClientPresenceActor::IsEditorOnly() const
+{
+	return bIsEditorOnly;
+}
+
 void AConcertClientPresenceActor::SetPresenceName(const FString& InName)
 {
 	if (PresenceTextComponent != nullptr)
diff --git a/Engine/Plugins/Developer/Concert/ConcertSync/ConcertSyncClient/Source/ConcertSyncClient/Private/ConcertClientPresenceActor.h b/Engine/Plugins/Developer/Concert/ConcertSync/ConcertSyncClient/Source/ConcertSyncClient/Private/ConcertClientPresenceActor.h
--- a/Engine/Plugins/Developer/Concert/ConcertSync/ConcertSyncClient/Source/ConcertSyncClient/Private/ConcertClientPresenceActor.h
+++ b/Engine/Plugins/Developer/Concert/ConcertSync/ConcertSyncClient/Source/ConcertSyncClient/Private/ConcertClientPresenceActor.h
@@ -26,10 +26,7 @@ class AConcertClientPresenceActor : public AActor
 public:
 
 	/** AActor interface */
-	virtual bool IsEditorOnly() const override final
-	{
-		return true;
-	}
+	virtual bool IsEditorOnly() const override final;
 
 #if WITH_EDITOR
 	virtual bool IsSelectable() const override final
@@ -77,5 +74,8 @@ protected:
 	class UMaterialInstanceDynamic* TextMID;
 
 	TOptional<FConcertClientMovement> PresenceMovement;
+
+	/** Flag initialized from presence settings on actor creation to have the presence potentially show up in PIE or Game. */
+	bool bIsEditorOnly = true;
 };
 
diff --git a/Engine/Plugins/Developer/Concert/ConcertSync/ConcertSyncClient/Source/ConcertSyncClient/Private/ConcertClientPresenceManager.cpp b/Engine/Plugins/Developer/Concert/ConcertSync/ConcertSyncClient/Source/ConcertSyncClient/Private/ConcertClientPresenceManager.cpp
--- a/Engine/Plugins/Developer/Concert/ConcertSync/ConcertSyncClient/Source/ConcertSyncClient/Private/ConcertClientPresenceManager.cpp
+++ b/Engine/Plugins/Developer/Concert/ConcertSync/ConcertSyncClient/Source/ConcertSyncClient/Private/ConcertClientPresenceManager.cpp
@@ -19,6 +19,8 @@
 #include "ConcertClientVRPresenceActor.h"
 #include "Scratchpad/ConcertScratchpad.h"
 #include "GameFramework/PlayerController.h"
+#include "Engine/GameEngine.h"
+#include "HAL/IConsoleManager.h"
 
 #if WITH_EDITOR
 #include "Editor.h"
@@ -56,7 +58,8 @@ bool ShowPresenceInPIE(const bool InIsPIE)
 
 }
 
-static TAutoConsoleVariable<int32> CVarEnablePresence(TEXT("concert.EnablePresence"), 1, TEXT("Enable Concert Presence"));
+static TAutoConsoleVariable<int32> CVarDisplayPresence(TEXT("Concert.DisplayPresence"), 1, TEXT("Enable display of Concert Presence from remote users."));
+static TAutoConsoleVariable<int32> CVarEmitPresence(TEXT("Concert.EmitPresence"), 1, TEXT("Enable display update of Concert Presence to remote users."));
 
 const TCHAR* FConcertClientPresenceManager::AssetContainerPath = TEXT("/ConcertSyncClient/ConcertAssets");
 
@@ -152,7 +155,7 @@ const UConcertAssetContainer& FConcertClientPresenceManager::GetAssetContainer()
 
 bool FConcertClientPresenceManager::IsPresenceVisible(const FConcertClientPresenceState& InPresenceState) const
 {
-	return InPresenceState.bVisible && ConcertClientPresenceManagerUtil::ShowPresenceInPIE(InPresenceState.EditorPlayMode == EEditorPlayMode::PIE);
+	return InPresenceState.bVisible && IsPIEPresenceEnabled(InPresenceState);
 }
 
 bool FConcertClientPresenceManager::IsPresenceVisible(const FGuid& InEndpointId) const
@@ -161,6 +164,11 @@ bool FConcertClientPresenceManager::IsPresenceVisible(const FGuid& InEndpointId)
 	return PresenceState && IsPresenceVisible(*PresenceState);
 }
 
+bool FConcertClientPresenceManager::IsPIEPresenceEnabled(const FConcertClientPresenceState& InPresenceState) const
+{
+	return ConcertClientPresenceManagerUtil::ShowPresenceInPIE(InPresenceState.EditorPlayMode == EEditorPlayMode::PIE);
+}
+
 template<class PresenceActorClass, typename PresenceUpdateEventType>
 void FConcertClientPresenceManager::UpdatePresence(AConcertClientPresenceActor* InPresenceActor, const PresenceUpdateEventType& InEvent)
 {
@@ -228,16 +236,23 @@ void FConcertClientPresenceManager::Unregister()
 	Session->UnregisterCustomEventHandler<FConcertPlaySessionEvent>(this);
 }
 
-UWorld* FConcertClientPresenceManager::GetWorld() const
+UWorld* FConcertClientPresenceManager::GetWorld(bool bIgnorePIE) const
 {
-	check(GEditor);
-
-	if (FWorldContext* WorldContext = GEditor->GetPIEWorldContext())
+	// Get the current world.
+	if (GIsEditor)
 	{
-		return WorldContext->World();
+		FWorldContext* PIEWorldContext = GEditor->GetPIEWorldContext();
+		if (!bIgnorePIE && PIEWorldContext)
+		{
+			return PIEWorldContext->World();
+		}
+		return GEditor->GetEditorWorldContext().World();
 	}
-
-	return GEditor->GetEditorWorldContext().World();
+	else if (UGameEngine* GameEngine = Cast<UGameEngine>(GEngine))
+	{
+		return GameEngine->GetGameWorld();
+	}
+	return nullptr;
 }
 
 FLevelEditorViewportClient* FConcertClientPresenceManager::GetPerspectiveViewport() const
@@ -254,7 +269,8 @@ void FConcertClientPresenceManager::OnEndFrame()
 	double DeltaTime = CurrentTime - PreviousEndFrameTime;
 	SecondsSinceLastLocationUpdate += DeltaTime;
 
-	if (SecondsSinceLastLocationUpdate >= ConcertClientPresenceManagerUtil::LocationUpdateFrequencySeconds)
+	// Game client do not generate a presence mode and state.
+	if (SecondsSinceLastLocationUpdate >= ConcertClientPresenceManagerUtil::LocationUpdateFrequencySeconds && !IsInGame() && CVarEmitPresence.GetValueOnAnyThread() > 0)
 	{
 		// Reload the avatar classes if the local client changed any of them.
 		if (UpdateLocalUserAvatarClasses())
@@ -313,11 +329,10 @@ void FConcertClientPresenceManager::SynchronizePresenceState()
 		FName EventWorldPathName;
 		const TSharedPtr<FConcertClientPresenceDataUpdateEvent> PresenceUpdateEvent = GetCachedPresenceState(PresenceState);
 		EventWorldPathName = PresenceUpdateEvent.IsValid() ? PresenceUpdateEvent->WorldPath : TEXT("");
-
-		const bool bIsValidViewport = GetPerspectiveViewport() != nullptr;
+		// NOTE: since presence state doesn't report PIE world path, bInCurrentWorld will never be true when in PIE, this will need to change if we want presence to display in PIE
 		const bool bInCurrentWorld = !ActiveWorldPathName.IsNone() && ActiveWorldPathName == EventWorldPathName;
 		
-		const bool bShowPresence = bIsPresenceEnabled && bIsValidViewport && bInCurrentWorld && PresenceState.bIsConnected && IsPresenceVisible(PresenceState) && (CVarEnablePresence.GetValueOnAnyThread() > 0);
+		const bool bShowPresence = bIsPresenceEnabled && bInCurrentWorld && PresenceState.bIsConnected && IsPresenceVisible(PresenceState) && (CVarDisplayPresence.GetValueOnAnyThread() > 0);
 		if (bShowPresence)
 		{
 			FConcertSessionClientInfo ClientSessionInfo;
@@ -327,8 +342,8 @@ void FConcertClientPresenceManager::SynchronizePresenceState()
 			{
 				PresenceState.PresenceActor = CreatePresenceActor(ClientSessionInfo.ClientInfo, PresenceState.VRDevice);
 			}
-			else if ((PresenceState.VRDevice.IsNone()  && PresenceState.PresenceActor->GetClass()->GetPathName() != ClientSessionInfo.ClientInfo.DesktopAvatarActorClass) || // Not in VR and Desktop Avatar changed?
-			         (!PresenceState.VRDevice.IsNone() && PresenceState.PresenceActor->GetClass()->GetPathName() != ClientSessionInfo.ClientInfo.VRAvatarActorClass))        // In VR and VR Avatar changed?
+			else if ((PresenceState.VRDevice.IsNone()  && PresenceState.PresenceActor->GetClass()->GetPathName() != ClientSessionInfo.ClientInfo.DesktopAvatarActorClass) // Not in VR and Desktop Avatar changed?
+			        || (!PresenceState.VRDevice.IsNone() && PresenceState.PresenceActor->GetClass()->GetPathName() != ClientSessionInfo.ClientInfo.VRAvatarActorClass))    // In VR and VR Avatar changed?					
 			{
 				ClearPresenceActor(RemoteEndpointId);
 				PresenceState.PresenceActor = CreatePresenceActor(ClientSessionInfo.ClientInfo, PresenceState.VRDevice);
@@ -644,16 +659,19 @@ void FConcertClientPresenceManager::SetPresenceVisibility(const FGuid& InEndpoin
 	}
 }
 
+bool FConcertClientPresenceManager::IsInGame() const
+{
+	return !GIsEditor;
+}
+
 bool FConcertClientPresenceManager::IsInPIE() const
 {
-	check(GEditor);
-	return GEditor->PlayWorld && !GEditor->bIsSimulatingInEditor;
+	return GEditor && GEditor->PlayWorld && !GEditor->bIsSimulatingInEditor;
 }
 
 bool FConcertClientPresenceManager::IsInSIE() const
 {
-	check(GEditor);
-	return GEditor->PlayWorld && GEditor->bIsSimulatingInEditor;
+	return GEditor && GEditor->PlayWorld && GEditor->bIsSimulatingInEditor;
 }
 
 void FConcertClientPresenceManager::SetPresenceEnabled(const bool bIsEnabled)
@@ -781,7 +799,7 @@ bool FConcertClientPresenceManager::IsJumpToPresenceEnabled(FGuid InEndpointId)
 		{
 			// The clients should be in the same world to enable teleporting. Note that both paths below are the from
 			// the editor world context. (The presence manager don't use PIE/SIE context world path).
-			return *GEditor->GetEditorWorldContext().World()->GetPathName() == CachedPresenceState->WorldPath;
+			return *GetWorld(/*bIgnorePIE=*/true)->GetPathName() == CachedPresenceState->WorldPath;
 		}
 	}
 
@@ -845,7 +863,7 @@ bool FConcertClientPresenceManager::IsShowHidePresenceEnabled(FGuid InEndpointId
 	}
 
 	const FConcertClientPresenceState* PresenceState = PresenceStateMap.Find(InEndpointId);
-	return PresenceState && ConcertClientPresenceManagerUtil::ShowPresenceInPIE(PresenceState->EditorPlayMode == EEditorPlayMode::PIE);
+	return PresenceState && IsPIEPresenceEnabled(*PresenceState);
 }
 
 FText FConcertClientPresenceManager::GetShowHidePresenceText(FGuid InEndpointId) const
@@ -896,7 +914,10 @@ FString FConcertClientPresenceManager::GetPresenceWorldPath(const FGuid& InEndpo
 		// and this leak into the world path. For presence management purpose, we want don't want to display the PIE/SIE context world path to
 		// the user and we want to be able to jump to the other presence when two users are visualizing the same level (comparing their world path)
 		// whatever the play mode is (PIE/SIE/None).
-		WorldPath = GEditor->GetEditorWorldContext().World()->GetPathName();
+		if (UWorld* CurrentWorld = GetWorld(/*bIgnorePIE=*/true))
+		{
+			WorldPath = CurrentWorld->GetPathName();
+		}
 	}
 	else
 	{
@@ -911,7 +932,6 @@ FString FConcertClientPresenceManager::GetPresenceWorldPath(const FGuid& InEndpo
 			}
 		}
 	}
-
 	return WorldPath;
 }
 
diff --git a/Engine/Plugins/Developer/Concert/ConcertSync/ConcertSyncClient/Source/ConcertSyncClient/Private/ConcertClientPresenceManager.h b/Engine/Plugins/Developer/Concert/ConcertSync/ConcertSyncClient/Source/ConcertSyncClient/Private/ConcertClientPresenceManager.h
--- a/Engine/Plugins/Developer/Concert/ConcertSync/ConcertSyncClient/Source/ConcertSyncClient/Private/ConcertClientPresenceManager.h
+++ b/Engine/Plugins/Developer/Concert/ConcertSync/ConcertSyncClient/Source/ConcertSyncClient/Private/ConcertClientPresenceManager.h
@@ -87,14 +87,21 @@ public:
 	/** Gets the container for all the assets of Concert clients. */
 	const class UConcertAssetContainer& GetAssetContainer() const;
 
+	/** Returns true if current session is in Game (i.e. !GIsEditor) */
+	bool IsInGame() const;
+
 	/** Returns true if current session is in PIE */
 	bool IsInPIE() const;
 
 	/** Returns true if current session is in SIE */
 	bool IsInSIE() const;
 
-	/** Get the current world */
-	UWorld* GetWorld() const;
+	/** 
+	 * Get the current world
+	 * @param bIgnorePIE if bIgnorePIE is true will still return the editor world even if an PIE world exists.
+	 * @return the current world
+	 */
+	UWorld* GetWorld(bool bIgnorePIE = false) const;
 
 	/** Get the active perspective viewport */
 	FLevelEditorViewportClient * GetPerspectiveViewport() const;
@@ -148,6 +155,9 @@ private:
 	/** Is presence visible for the given state? */
 	bool IsPresenceVisible(const FConcertClientPresenceState& InPresenceState) const;
 
+	/** Return if PIE presence is potentially visible for a state. */
+	bool IsPIEPresenceEnabled(const FConcertClientPresenceState& InPresenceState) const;
+
 	/** Is presence visible for the given endpoint id? */
 	bool IsPresenceVisible(const FGuid& InEndpointId) const;
 
diff --git a/Engine/Plugins/Developer/Concert/ConcertSync/ConcertSyncClient/Source/ConcertSyncClient/Private/ConcertSyncClient.cpp b/Engine/Plugins/Developer/Concert/ConcertSync/ConcertSyncClient/Source/ConcertSyncClient/Private/ConcertSyncClient.cpp
--- a/Engine/Plugins/Developer/Concert/ConcertSync/ConcertSyncClient/Source/ConcertSyncClient/Private/ConcertSyncClient.cpp
+++ b/Engine/Plugins/Developer/Concert/ConcertSync/ConcertSyncClient/Source/ConcertSyncClient/Private/ConcertSyncClient.cpp
@@ -149,13 +149,10 @@ void FConcertSyncClient::RegisterConcertSyncHandlers(TSharedRef<IConcertClientSe
 	{
 		CreateWorkspace(LiveSession.ToSharedRef());
 #if WITH_EDITOR
-		if (GIsEditor)
+		PresenceManager.Reset();
+		if (EnumHasAnyFlags(SessionFlags, EConcertSyncSessionFlags::EnablePresence))
 		{
-			PresenceManager.Reset();
-			if (EnumHasAnyFlags(SessionFlags, EConcertSyncSessionFlags::EnablePresence))
-			{
-				PresenceManager = MakeShared<FConcertClientPresenceManager>(InSession); // TODO: Use LiveSession?
-			}
+			PresenceManager = MakeShared<FConcertClientPresenceManager>(InSession); // TODO: Use LiveSession?
 		}
 		if (EnumHasAnyFlags(SessionFlags, EConcertSyncSessionFlags::EnableSequencer))
 		{
