- Don't record saved package event if the long package name is invalid (/Temp for example)

diff --git a/Engine/Plugins/Developer/Concert/ConcertSync/ConcertSyncClient/Source/ConcertSyncClient/Private/ConcertClientPackageManager.cpp b/Engine/Plugins/Developer/Concert/ConcertSync/ConcertSyncClient/Source/ConcertSyncClient/Private/ConcertClientPackageManager.cpp
--- a/Engine/Plugins/Developer/Concert/ConcertSync/ConcertSyncClient/Source/ConcertSyncClient/Private/ConcertClientPackageManager.cpp
+++ b/Engine/Plugins/Developer/Concert/ConcertSync/ConcertSyncClient/Source/ConcertSyncClient/Private/ConcertClientPackageManager.cpp
@@ -310,6 +310,11 @@ void FConcertClientPackageManager::HandleLocalPackageEvent(const FConcertPackage
 				return;
 			}
 		}
+
+		if (!FPackageName::IsValidLongPackageName(Package.Info.PackageName.ToString())) // Auto-Save might save the template in /Temp/... which is an invalid long package name.
+		{
+			return;
+		}
 	}
 
 	if (Package.Info.PackageUpdateType == EConcertPackageUpdateType::Added && EnumHasAnyFlags(LiveSession->GetSessionFlags(), EConcertSyncSessionFlags::ShouldUsePackageSandbox))
