Edigrated 9459835

  - Don't purge/reload 'compiled in' packages.
  - Don't reload 'in memory only' packages.

diff --git a/Engine/Plugins/Developer/Concert/ConcertSync/ConcertSyncClient/Source/ConcertSyncClient/Private/ConcertClientPackageManager.cpp b/Engine/Plugins/Developer/Concert/ConcertSync/ConcertSyncClient/Source/ConcertSyncClient/Private/ConcertClientPackageManager.cpp
--- a/Engine/Plugins/Developer/Concert/ConcertSync/ConcertSyncClient/Source/ConcertSyncClient/Private/ConcertClientPackageManager.cpp
+++ b/Engine/Plugins/Developer/Concert/ConcertSync/ConcertSyncClient/Source/ConcertSyncClient/Private/ConcertClientPackageManager.cpp
@@ -97,7 +97,7 @@ bool FConcertClientPackageManager::ShouldIgnorePackageDirtyEvent(class UPackage*
 {
 	return InPackage == GetTransientPackage()
 		|| InPackage->HasAnyFlags(RF_Transient)
-		|| InPackage->HasAnyPackageFlags(PKG_PlayInEditor)
+		|| InPackage->HasAnyPackageFlags(PKG_PlayInEditor | PKG_CompiledIn) // CompiledIn packages are not considered content for MU. (ex when changing some plugin settings like /Script/DisasterRecoveryClient)
 		|| bIgnorePackageDirtyEvent;
 }
 
@@ -278,7 +278,7 @@ void FConcertClientPackageManager::HandlePackageRejectedEvent(const FConcertSess
 void FConcertClientPackageManager::HandlePackageDirtyStateChanged(UPackage* InPackage)
 {
 	check(!InPackage->HasAnyFlags(RF_Transient) || InPackage != GetTransientPackage());
-	if (InPackage->IsDirty())
+	if (InPackage->IsDirty() && !InPackage->HasAnyPackageFlags(PKG_CompiledIn | PKG_InMemoryOnly)) // Dirty packages are tracked for purge/reload, but 'compiled in' and 'in memory' cannot be hot purged/reloaded.
 	{
 		DirtyPackages.Add(InPackage->GetFName());
 	}
diff --git a/Engine/Plugins/Developer/Concert/ConcertSync/ConcertSyncClient/Source/ConcertSyncClient/Private/ConcertClientTransactionBridge.cpp b/Engine/Plugins/Developer/Concert/ConcertSync/ConcertSyncClient/Source/ConcertSyncClient/Private/ConcertClientTransactionBridge.cpp
--- a/Engine/Plugins/Developer/Concert/ConcertSync/ConcertSyncClient/Source/ConcertSyncClient/Private/ConcertClientTransactionBridge.cpp
+++ b/Engine/Plugins/Developer/Concert/ConcertSync/ConcertSyncClient/Source/ConcertSyncClient/Private/ConcertClientTransactionBridge.cpp
@@ -59,8 +59,8 @@ bool RunTransactionFilters(const TArray<FTransactionClassFilter>& InFilters, UOb
 
 ETransactionFilterResult ApplyTransactionFilters(UObject* InObject, UPackage* InChangedPackage)
 {
-	// Ignore transient packages and objects
-	if (!InChangedPackage || InChangedPackage == GetTransientPackage() || InChangedPackage->HasAnyFlags(RF_Transient) || InObject->HasAnyFlags(RF_Transient))
+	// Ignore transient packages and objects, compiled in package are not considered Multi-user content.
+	if (!InChangedPackage || InChangedPackage == GetTransientPackage() || InChangedPackage->HasAnyFlags(RF_Transient) || InChangedPackage->HasAnyPackageFlags(PKG_CompiledIn) || InObject->HasAnyFlags(RF_Transient))
 	{
 		return ETransactionFilterResult::ExcludeObject;
 	}
