Edigrated 9583111

  - Prevented access to the client name functor when populating the search string in DisasterRecovery because the client name is not displayed/available in this mode.

diff --git a/Engine/Plugins/Developer/Concert/ConcertSync/ConcertSyncClient/Source/ConcertSyncUI/Private/SConcertSessionActivities.cpp b/Engine/Plugins/Developer/Concert/ConcertSync/ConcertSyncClient/Source/ConcertSyncUI/Private/SConcertSessionActivities.cpp
--- a/Engine/Plugins/Developer/Concert/ConcertSync/ConcertSyncClient/Source/ConcertSyncUI/Private/SConcertSessionActivities.cpp
+++ b/Engine/Plugins/Developer/Concert/ConcertSync/ConcertSyncClient/Source/ConcertSyncUI/Private/SConcertSessionActivities.cpp
@@ -821,7 +821,7 @@ FText SConcertSessionActivities::UpdateTextFilter(const FText& InFilterText)
 
 void SConcertSessionActivities::PopulateSearchStrings(const FConcertClientSessionActivity& Activity, TArray<FString>& OutSearchStrings) const
 {
-	FText ClientName = ConcertSessionActivityUtils::GetClientName(GetActivityUserFn(Activity.Activity.EndpointId));
+	FText ClientName = GetActivityUserFn ? ConcertSessionActivityUtils::GetClientName(GetActivityUserFn(Activity.Activity.EndpointId)) : FText::GetEmpty();
 
 	OutSearchStrings.Add(ConcertSessionActivityUtils::GetActivityDateTime(Activity, TimeFormat.Get()).ToString());
 	OutSearchStrings.Add(ConcertSessionActivityUtils::GetSummary(Activity, ClientName, /*bAsRichText*/false).ToString());
