Editgrate cl 9542521:

Fix crash in -game when getting AssetEditorSubsystem in Multi-User

diff --git a/Engine/Plugins/Developer/Concert/ConcertSync/ConcertSyncClient/Source/ConcertSyncClient/Private/ConcertClientSequencerManager.cpp b/Engine/Plugins/Developer/Concert/ConcertSync/ConcertSyncClient/Source/ConcertSyncClient/Private/ConcertClientSequencerManager.cpp
--- a/Engine/Plugins/Developer/Concert/ConcertSync/ConcertSyncClient/Source/ConcertSyncClient/Private/ConcertClientSequencerManager.cpp
+++ b/Engine/Plugins/Developer/Concert/ConcertSync/ConcertSyncClient/Source/ConcertSyncClient/Private/ConcertClientSequencerManager.cpp
@@ -327,12 +327,12 @@ void FConcertClientSequencerManager::OnOpenEvent(const FConcertSessionContext&,
 void FConcertClientSequencerManager::ApplyTransportOpenEvent(const FString& SequenceObjectPath)
 {
 	TGuardValue<bool> ReentrancyGuard(bRespondingToTransportEvent, true);
-	if (IsSequencerRemoteOpenEnabled())
-	{
 #if WITH_EDITOR
+	if (IsSequencerRemoteOpenEnabled() && GIsEditor)
+	{
 		GEditor->GetEditorSubsystem<UAssetEditorSubsystem>()->OpenEditorForAsset(SequenceObjectPath);
-#endif
 	}
+#endif
 }
 
 void FConcertClientSequencerManager::ApplyCloseToPlayers(const FConcertSequencerCloseEvent& InEvent)
diff --git a/Engine/Plugins/Developer/Concert/ConcertSync/ConcertSyncClient/Source/ConcertSyncClient/Private/ConcertSyncClientUtil.cpp b/Engine/Plugins/Developer/Concert/ConcertSync/ConcertSyncClient/Source/ConcertSyncClient/Private/ConcertSyncClientUtil.cpp
--- a/Engine/Plugins/Developer/Concert/ConcertSync/ConcertSyncClient/Source/ConcertSyncClient/Private/ConcertSyncClientUtil.cpp
+++ b/Engine/Plugins/Developer/Concert/ConcertSync/ConcertSyncClient/Source/ConcertSyncClient/Private/ConcertSyncClientUtil.cpp
@@ -34,11 +34,10 @@
 	#include "Editor/UnrealEdEngine.h"
 	#include "PackageTools.h"
 	#include "ObjectTools.h"
-	
+	#include "Subsystems/AssetEditorSubsystem.h"
 	#include "GameMapsSettings.h"
 	#include "FileHelpers.h"
 #endif
-#include "Subsystems/AssetEditorSubsystem.h"
 
 namespace ConcertSyncClientUtil
 {
@@ -356,7 +355,7 @@ void PurgePackages(TArrayView<const FName> InPackageNames)
 	TArray<UObject*> ObjectsToPurge;
 	auto CollectObjectToPurge = [&ObjectsToPurge](UObject* InObject)
 	{
-		if (InObject->IsAsset())
+		if (InObject->IsAsset() && GIsEditor)
 		{
 			GEditor->GetEditorSubsystem<UAssetEditorSubsystem>()->CloseAllEditorsForAsset(InObject);
 		}
