Skip transient packages when tracking saves and discards

#rnx

diff --git a/Engine/Plugins/Developer/Concert/ConcertSync/ConcertSyncClient/Source/ConcertSyncClient/Private/ConcertClientPackageBridge.cpp b/Engine/Plugins/Developer/Concert/ConcertSync/ConcertSyncClient/Source/ConcertSyncClient/Private/ConcertClientPackageBridge.cpp
--- a/Engine/Plugins/Developer/Concert/ConcertSync/ConcertSyncClient/Source/ConcertSyncClient/Private/ConcertClientPackageBridge.cpp
+++ b/Engine/Plugins/Developer/Concert/ConcertSync/ConcertSyncClient/Source/ConcertSyncClient/Private/ConcertClientPackageBridge.cpp
@@ -33,6 +33,23 @@ void FillPackageInfo(UPackage* InPackage, const EConcertPackageUpdateType InPack
 	OutPackageInfo.PackageUpdateType = InPackageUpdateType;
 }
 
+bool ShouldIgnorePackage(const UPackage* InPackage)
+{
+	// Ignore transient packages and objects, compiled in package are not considered Multi-user content.
+	if (!InPackage || InPackage == GetTransientPackage() || InPackage->HasAnyFlags(RF_Transient) || InPackage->HasAnyPackageFlags(PKG_CompiledIn))
+	{
+		return true;
+	}
+
+	// Ignore packages outside of known root paths (we ignore read-only roots here to skip things like unsaved worlds)
+	if (!FPackageName::IsValidLongPackageName(InPackage->GetName()))
+	{
+		return true;
+	}
+
+	return false;
+}
+
 } // namespace ConcertClientPackageBridgeUtil
 
 FConcertClientPackageBridge::FConcertClientPackageBridge()
@@ -117,7 +134,7 @@ void FConcertClientPackageBridge::HandlePackagePreSave(UPackage* Package)
 	}
 
 	// Ignore unwanted saves
-	if (bIgnoreLocalSave)
+	if (bIgnoreLocalSave || ConcertClientPackageBridgeUtil::ShouldIgnorePackage(Package))
 	{
 		return;
 	}
@@ -159,7 +176,7 @@ void FConcertClientPackageBridge::HandlePackageSaved(const FString& PackageFilen
 	}
 
 	// Ignore unwanted saves
-	if (bIgnoreLocalSave)
+	if (bIgnoreLocalSave || ConcertClientPackageBridgeUtil::ShouldIgnorePackage(Package))
 	{
 		return;
 	}
@@ -268,9 +285,12 @@ void FConcertClientPackageBridge::HandleAssetReload(const EPackageReloadPhase In
 	if (InPackageReloadPhase == EPackageReloadPhase::PrePackageLoad)
 	{
 		UPackage* Package = const_cast<UPackage*>(InPackageReloadedEvent->GetOldPackage());
-		OnLocalPackageDiscardedDelegate.Broadcast(Package);
+		if (!ConcertClientPackageBridgeUtil::ShouldIgnorePackage(Package))
+		{
+			OnLocalPackageDiscardedDelegate.Broadcast(Package);
 
-		UE_LOG(LogConcert, Verbose, TEXT("Asset Discarded: %s"), *Package->GetName());
+			UE_LOG(LogConcert, Verbose, TEXT("Asset Discarded: %s"), *Package->GetName());
+		}
 	}
 }
 
@@ -285,9 +305,12 @@ void FConcertClientPackageBridge::HandleMapChanged(UWorld* InWorld, EMapChangeTy
 	if (InMapChangeType == EMapChangeType::TearDownWorld)
 	{
 		UPackage* Package = InWorld->GetOutermost();
-		OnLocalPackageDiscardedDelegate.Broadcast(Package);
+		if (!ConcertClientPackageBridgeUtil::ShouldIgnorePackage(Package))
+		{
+			OnLocalPackageDiscardedDelegate.Broadcast(Package);
 
-		UE_LOG(LogConcert, Verbose, TEXT("Asset Discarded: %s"), *Package->GetName());
+			UE_LOG(LogConcert, Verbose, TEXT("Asset Discarded: %s"), *Package->GetName());
+		}
 	}
 }
 
