Fixed Multi-User Presence not using the proper avatar if VR was entered before joining a session

#rnx

diff --git a/Engine/Plugins/Developer/Concert/ConcertSync/ConcertSyncClient/Source/ConcertSyncClient/Private/ConcertClientPresenceManager.cpp b/Engine/Plugins/Developer/Concert/ConcertSync/ConcertSyncClient/Source/ConcertSyncClient/Private/ConcertClientPresenceManager.cpp
--- a/Engine/Plugins/Developer/Concert/ConcertSync/ConcertSyncClient/Source/ConcertSyncClient/Private/ConcertClientPresenceManager.cpp
+++ b/Engine/Plugins/Developer/Concert/ConcertSync/ConcertSyncClient/Source/ConcertSyncClient/Private/ConcertClientPresenceManager.cpp
@@ -904,6 +904,10 @@ FConcertClientDefaultPresenceModeFactory::FConcertClientDefaultPresenceModeFacto
 	// Add handler for VR mode
 	IVREditorModule::Get().OnVREditingModeEnter().AddRaw(this, &FConcertClientDefaultPresenceModeFactory::OnVREditingModeEnter);
 	IVREditorModule::Get().OnVREditingModeExit().AddRaw(this, &FConcertClientDefaultPresenceModeFactory::OnVREditingModeExit);
+
+	// Set the initial vr device if any
+	UVREditorMode* VRMode = IVREditorModule::Get().GetVRMode();
+	VRDeviceType = VRMode ? VRMode->GetHMDDeviceType() : FName();
 }
 
 FConcertClientDefaultPresenceModeFactory::~FConcertClientDefaultPresenceModeFactory()
