- Fixed DisasterRecovery and MultiUser closing a FScopedSlowTask out of order when performing initial workspace sync at the same time.

diff --git a/Engine/Plugins/Developer/Concert/ConcertSync/ConcertSyncClient/Source/ConcertSyncClient/Private/ConcertClientWorkspace.cpp b/Engine/Plugins/Developer/Concert/ConcertSync/ConcertSyncClient/Source/ConcertSyncClient/Private/ConcertClientWorkspace.cpp
--- a/Engine/Plugins/Developer/Concert/ConcertSync/ConcertSyncClient/Source/ConcertSyncClient/Private/ConcertClientWorkspace.cpp
+++ b/Engine/Plugins/Developer/Concert/ConcertSync/ConcertSyncClient/Source/ConcertSyncClient/Private/ConcertClientWorkspace.cpp
@@ -55,6 +55,62 @@
 
 #define LOCTEXT_NAMESPACE "ConcertClientWorkspace"
 
+/** Provides a workaround for scoped slow tasks that are push/pop out of order by Concert. Concert doesn't use FScopedSlowTask as designed when syncing
+    workspaces. The tasks are likely to run over many engine loops. DisasterRecovery and MultiUser, based on Concert, may both try to run the same
+    slow tasks to sync their workspaces. Syncing uses network messages and the duration depends of how much needs to be synced, so those slow tasks ends
+    in an arbitrary order rather than the strict order expected by the FSlowTask design, sometimes triggering an ensure(). */
+class FConcertSlowTaskStackWorkaround
+{
+public:
+	static FConcertSlowTaskStackWorkaround& Get()
+	{
+		static FConcertSlowTaskStackWorkaround Instance; // Static instance shared by DisasterRecovery and MultiUser clients.
+		return Instance;
+	}
+
+	void PushTask(TUniquePtr<FScopedSlowTask>& Task)
+	{
+		TaskStack.Push(Task.Get()); // Track the expected order as FScopedSlowTask does.
+	}
+
+	void PopTask(TUniquePtr<FScopedSlowTask> Discarded)
+	{
+		if (!Discarded)
+		{
+			return;
+		}
+
+		if (Discarded.Get() == TaskStack.Top()) // If the slow task completed in the FScopedSlowTask expected order?
+		{
+			TaskStack.Pop();
+			Discarded.Reset();
+			while (TaskStack.Num() && ExtendedTaskLife.Num()) // Check if a task for which the life was extended is now on top (and can be discarded).
+			{
+				FScopedSlowTask* Top = TaskStack.Top();
+				if (ExtendedTaskLife.RemoveAll([Top](const TUniquePtr<FScopedSlowTask>& Task) { return Task.Get() == Top; }) != 0)
+				{
+					TaskStack.Pop();
+				}
+				else
+				{
+					return;
+				}
+			}
+		}
+		else // Out of order deletion, need to extend the life of this task until it gets to the top of the stack.
+		{
+			ExtendedTaskLife.Add(MoveTemp(Discarded));
+		}
+	}
+
+private:
+	FConcertSlowTaskStackWorkaround() = default;
+
+	TArray<FScopedSlowTask*> TaskStack;
+	TArray<TUniquePtr<FScopedSlowTask>> ExtendedTaskLife;
+};
+
+
 FConcertClientWorkspace::FConcertClientWorkspace(TSharedRef<FConcertSyncClientLiveSession> InLiveSession, IConcertClientPackageBridge* InPackageBridge, IConcertClientTransactionBridge* InTransactionBridge)
 {
 	BindSession(InLiveSession, InPackageBridge, InTransactionBridge);
@@ -493,6 +549,7 @@ void FConcertClientWorkspace::HandleConnectionChanged(IConcertClientSession& InS
 		bFinalizeWorkspaceSyncRequested = false;
 		InitialSyncSlowTask = MakeUnique<FScopedSlowTask>(1.0f, LOCTEXT("SynchronizingSession", "Synchronizing Session..."));
 		InitialSyncSlowTask->MakeDialog();
+		FConcertSlowTaskStackWorkaround::Get().PushTask(InitialSyncSlowTask);
 
 		// Request our initial workspace sync for any new activity since we last joined
 		{
@@ -526,7 +583,7 @@ void FConcertClientWorkspace::HandleConnectionChanged(IConcertClientSession& InS
 	{
 		bHasSyncedWorkspace = false;
 		bFinalizeWorkspaceSyncRequested = false;
-		InitialSyncSlowTask.Reset();
+		FConcertSlowTaskStackWorkaround::Get().PopTask(MoveTemp(InitialSyncSlowTask));
 	}
 }
 
@@ -712,7 +769,7 @@ void FConcertClientWorkspace::OnEndFrame()
 
 		// Finalize the sync
 		bHasSyncedWorkspace = true;
-		InitialSyncSlowTask.Reset();
+		FConcertSlowTaskStackWorkaround::Get().PopTask(MoveTemp(InitialSyncSlowTask));
 	}
 
 	if (bHasSyncedWorkspace)
