Adding includes for incremental compiles

diff --git a/Engine/Plugins/Developer/Concert/ConcertSync/ConcertSyncClient/Source/ConcertSyncClient/Private/ConcertClientSequencerManager.cpp b/Engine/Plugins/Developer/Concert/ConcertSync/ConcertSyncClient/Source/ConcertSyncClient/Private/ConcertClientSequencerManager.cpp
--- a/Engine/Plugins/Developer/Concert/ConcertSync/ConcertSyncClient/Source/ConcertSyncClient/Private/ConcertClientSequencerManager.cpp
+++ b/Engine/Plugins/Developer/Concert/ConcertSync/ConcertSyncClient/Source/ConcertSyncClient/Private/ConcertClientSequencerManager.cpp
@@ -18,9 +18,9 @@
 #if WITH_EDITOR
 	#include "ISequencerModule.h"
 	#include "ISequencer.h"
-	
+	#include "Subsystems/AssetEditorSubsystem.h"
+	#include "Editor.h"
 #endif
-#include "Subsystems/AssetEditorSubsystem.h"
 
 DEFINE_LOG_CATEGORY_STATIC(LogConcertSequencerSync, Warning, Log)
 
@@ -316,7 +316,9 @@ void FConcertClientSequencerManager::ApplyTransportOpenEvent(const FString& Sequ
 	TGuardValue<bool> ReentrancyGuard(bRespondingToTransportEvent, true);
 	if (IsSequencerRemoteOpenEnabled())
 	{
+#if WITH_EDITOR
 		GEditor->GetEditorSubsystem<UAssetEditorSubsystem>()->OpenEditorForAsset(SequenceObjectPath);
+#endif
 	}
 }
 
diff --git a/Engine/Plugins/Editor/AssetManagerEditor/Source/AssetManagerEditor/Private/SAssetAuditBrowser.cpp b/Engine/Plugins/Editor/AssetManagerEditor/Source/AssetManagerEditor/Private/SAssetAuditBrowser.cpp
--- a/Engine/Plugins/Editor/AssetManagerEditor/Source/AssetManagerEditor/Private/SAssetAuditBrowser.cpp
+++ b/Engine/Plugins/Editor/AssetManagerEditor/Source/AssetManagerEditor/Private/SAssetAuditBrowser.cpp
@@ -33,6 +33,7 @@
 #include "DragAndDrop/AssetDragDropOp.h"
 #include "Blueprint/BlueprintSupport.h"
 #include "Subsystems/AssetEditorSubsystem.h"
+#include "Editor.h"
 
 #define LOCTEXT_NAMESPACE "AssetManagementBrowser"
 
diff --git a/Engine/Plugins/Editor/GameplayTagsEditor/Source/GameplayTagsEditor/Private/GameplayTagReferenceHelperDetails.cpp b/Engine/Plugins/Editor/GameplayTagsEditor/Source/GameplayTagsEditor/Private/GameplayTagReferenceHelperDetails.cpp
--- a/Engine/Plugins/Editor/GameplayTagsEditor/Source/GameplayTagsEditor/Private/GameplayTagReferenceHelperDetails.cpp
+++ b/Engine/Plugins/Editor/GameplayTagsEditor/Source/GameplayTagsEditor/Private/GameplayTagReferenceHelperDetails.cpp
@@ -12,6 +12,7 @@
 #include "SGameplayTagWidget.h"
 #include "IDetailChildrenBuilder.h"
 #include "Subsystems/AssetEditorSubsystem.h"
+#include "Editor.h"
 
 #define LOCTEXT_NAMESPACE "GameplayTagReferenceHelperDetails"
 
diff --git a/Engine/Plugins/Experimental/DataPrepEditor/Source/DataPrepCore/Private/DataprepCoreUtils.cpp b/Engine/Plugins/Experimental/DataPrepEditor/Source/DataPrepCore/Private/DataprepCoreUtils.cpp
--- a/Engine/Plugins/Experimental/DataPrepEditor/Source/DataPrepCore/Private/DataprepCoreUtils.cpp
+++ b/Engine/Plugins/Experimental/DataPrepEditor/Source/DataPrepCore/Private/DataprepCoreUtils.cpp
@@ -6,9 +6,10 @@
 
 #if WITH_EDITOR
 #include "ObjectTools.h"
-
-#endif
 #include "Subsystems/AssetEditorSubsystem.h"
+#include "Editor.h"
+#endif
+
 
 void FDataprepCoreUtils::PurgeObjects(TArray<UObject*> Objects)
 {
diff --git a/Engine/Plugins/FX/Niagara/Source/NiagaraEditor/Private/Widgets/SNiagaraScriptGraph.cpp b/Engine/Plugins/FX/Niagara/Source/NiagaraEditor/Private/Widgets/SNiagaraScriptGraph.cpp
--- a/Engine/Plugins/FX/Niagara/Source/NiagaraEditor/Private/Widgets/SNiagaraScriptGraph.cpp
+++ b/Engine/Plugins/FX/Niagara/Source/NiagaraEditor/Private/Widgets/SNiagaraScriptGraph.cpp
@@ -27,6 +27,7 @@
 #include "Framework/Application/SlateApplication.h"
 #include "GraphEditorActions.h"
 #include "Subsystems/AssetEditorSubsystem.h"
+#include "Editor.h"
 
 #define LOCTEXT_NAMESPACE "NiagaraScriptGraph"
 
diff --git a/Engine/Plugins/FX/Niagara/Source/NiagaraEditorWidgets/Private/Stack/SNiagaraStackModuleItem.cpp b/Engine/Plugins/FX/Niagara/Source/NiagaraEditorWidgets/Private/Stack/SNiagaraStackModuleItem.cpp
--- a/Engine/Plugins/FX/Niagara/Source/NiagaraEditorWidgets/Private/Stack/SNiagaraStackModuleItem.cpp
+++ b/Engine/Plugins/FX/Niagara/Source/NiagaraEditorWidgets/Private/Stack/SNiagaraStackModuleItem.cpp
@@ -26,6 +26,7 @@
 #include "NiagaraEditorUtilities.h"
 #include "SGraphActionMenu.h"
 #include "Subsystems/AssetEditorSubsystem.h"
+#include "Editor.h"
 
 #define LOCTEXT_NAMESPACE "NiagaraStackModuleItem"
 
diff --git a/Engine/Plugins/MovieScene/LevelSequenceEditor/Source/LevelSequenceEditor/Private/Misc/LevelSequenceEditorHelpers.cpp b/Engine/Plugins/MovieScene/LevelSequenceEditor/Source/LevelSequenceEditor/Private/Misc/LevelSequenceEditorHelpers.cpp
--- a/Engine/Plugins/MovieScene/LevelSequenceEditor/Source/LevelSequenceEditor/Private/Misc/LevelSequenceEditorHelpers.cpp
+++ b/Engine/Plugins/MovieScene/LevelSequenceEditor/Source/LevelSequenceEditor/Private/Misc/LevelSequenceEditorHelpers.cpp
@@ -34,6 +34,7 @@
 #include "PropertyEditorModule.h"
 #include "Widgets/Docking/SDockTab.h"
 #include "Subsystems/AssetEditorSubsystem.h"
+#include "Editor.h"
 
 /* LevelSequenceEditorHelpers
  *****************************************************************************/
diff --git a/Engine/Source/Editor/Cascade/Private/Tests/ParticleSystemTests.cpp b/Engine/Source/Editor/Cascade/Private/Tests/ParticleSystemTests.cpp
--- a/Engine/Source/Editor/Cascade/Private/Tests/ParticleSystemTests.cpp
+++ b/Engine/Source/Editor/Cascade/Private/Tests/ParticleSystemTests.cpp
@@ -17,7 +17,9 @@
 #include "Particles/ParticleEmitter.h"
 #include "Particles/ParticleLODLevel.h"
 #include "Particles/Size/ParticleModuleSize.h"
+
 #include "Subsystems/AssetEditorSubsystem.h"
+#include "Editor.h"
 
 
 #if WITH_DEV_AUTOMATION_TESTS
diff --git a/Engine/Source/Editor/GameProjectGeneration/Private/SNewClassDialog.cpp b/Engine/Source/Editor/GameProjectGeneration/Private/SNewClassDialog.cpp
--- a/Engine/Source/Editor/GameProjectGeneration/Private/SNewClassDialog.cpp
+++ b/Engine/Source/Editor/GameProjectGeneration/Private/SNewClassDialog.cpp
@@ -39,6 +39,7 @@
 #include "Widgets/Notifications/SNotificationList.h"
 #include "FeaturedClasses.inl"
 #include "Subsystems/AssetEditorSubsystem.h"
+#include "Editor.h"
 
 #define LOCTEXT_NAMESPACE "GameProjectGeneration"
 
diff --git a/Engine/Source/Editor/IntroTutorials/Private/STutorialRoot.cpp b/Engine/Source/Editor/IntroTutorials/Private/STutorialRoot.cpp
--- a/Engine/Source/Editor/IntroTutorials/Private/STutorialRoot.cpp
+++ b/Engine/Source/Editor/IntroTutorials/Private/STutorialRoot.cpp
@@ -18,6 +18,7 @@
 #include "Framework/Notifications/NotificationManager.h"
 #include "Widgets/Notifications/SNotificationList.h"
 #include "Subsystems/AssetEditorSubsystem.h"
+#include "Editor.h"
 
 #define LOCTEXT_NAMESPACE "STutorialRoot"
 
diff --git a/Engine/Source/Editor/IntroTutorials/Private/TutorialText.cpp b/Engine/Source/Editor/IntroTutorials/Private/TutorialText.cpp
--- a/Engine/Source/Editor/IntroTutorials/Private/TutorialText.cpp
+++ b/Engine/Source/Editor/IntroTutorials/Private/TutorialText.cpp
@@ -24,6 +24,7 @@
 #include "Interfaces/IAnalyticsProvider.h"
 #include "TutorialHyperlinkDecorator.h"
 #include "Subsystems/AssetEditorSubsystem.h"
+#include "Editor.h"
 
 #define LOCTEXT_NAMESPACE "TutorialText"
 
diff --git a/Engine/Source/Editor/MovieSceneTools/Private/Sections/CinematicShotSection.cpp b/Engine/Source/Editor/MovieSceneTools/Private/Sections/CinematicShotSection.cpp
--- a/Engine/Source/Editor/MovieSceneTools/Private/Sections/CinematicShotSection.cpp
+++ b/Engine/Source/Editor/MovieSceneTools/Private/Sections/CinematicShotSection.cpp
@@ -20,6 +20,7 @@
 #include "Evaluation/MovieSceneEvaluationTemplateInstance.h"
 #include "CommonMovieSceneTools.h"
 #include "Subsystems/AssetEditorSubsystem.h"
+#include "Editor.h"
 
 #define LOCTEXT_NAMESPACE "FCinematicShotSection"
 
diff --git a/Engine/Source/Editor/MovieSceneTools/Private/TrackEditors/SkeletalAnimationTrackEditor.cpp b/Engine/Source/Editor/MovieSceneTools/Private/TrackEditors/SkeletalAnimationTrackEditor.cpp
--- a/Engine/Source/Editor/MovieSceneTools/Private/TrackEditors/SkeletalAnimationTrackEditor.cpp
+++ b/Engine/Source/Editor/MovieSceneTools/Private/TrackEditors/SkeletalAnimationTrackEditor.cpp
@@ -45,6 +45,7 @@
 
 #include "CommonMovieSceneTools.h"
 #include "Subsystems/AssetEditorSubsystem.h"
+#include "Editor.h"
 
 namespace SkeletalAnimationEditorConstants
 {
diff --git a/Engine/Source/Editor/UMGEditor/Private/Templates/WidgetTemplateBlueprintClass.cpp b/Engine/Source/Editor/UMGEditor/Private/Templates/WidgetTemplateBlueprintClass.cpp
--- a/Engine/Source/Editor/UMGEditor/Private/Templates/WidgetTemplateBlueprintClass.cpp
+++ b/Engine/Source/Editor/UMGEditor/Private/Templates/WidgetTemplateBlueprintClass.cpp
@@ -10,6 +10,7 @@
 #include "Kismet2/BlueprintEditorUtils.h"
 #include "Styling/SlateIconFinder.h"
 #include "Subsystems/AssetEditorSubsystem.h"
+#include "Editor.h"
 
 #define LOCTEXT_NAMESPACE "UMGEditor"
 
diff --git a/Engine/Source/Editor/UnrealEd/Private/EditorAnimUtils.cpp b/Engine/Source/Editor/UnrealEd/Private/EditorAnimUtils.cpp
--- a/Engine/Source/Editor/UnrealEd/Private/EditorAnimUtils.cpp
+++ b/Engine/Source/Editor/UnrealEd/Private/EditorAnimUtils.cpp
@@ -18,6 +18,7 @@
 #include "IContentBrowserSingleton.h"
 #include "ContentBrowserModule.h"
 #include "Subsystems/AssetEditorSubsystem.h"
+#include "Editor.h"
 
 #define LOCTEXT_NAMESPACE "EditorAnimUtils"
 
diff --git a/Engine/Source/Editor/UnrealEd/Private/Fbx/ReimportFbxSceneFactory.cpp b/Engine/Source/Editor/UnrealEd/Private/Fbx/ReimportFbxSceneFactory.cpp
--- a/Engine/Source/Editor/UnrealEd/Private/Fbx/ReimportFbxSceneFactory.cpp
+++ b/Engine/Source/Editor/UnrealEd/Private/Fbx/ReimportFbxSceneFactory.cpp
@@ -56,6 +56,7 @@
 
 #include "AI/Navigation/NavCollisionBase.h"
 #include "Subsystems/AssetEditorSubsystem.h"
+#include "Editor.h"
 
 #define LOCTEXT_NAMESPACE "FBXSceneReImportFactory"
 
diff --git a/Engine/Source/Editor/UnrealEd/Private/Subsystems/AssetEditorSubsystem.cpp b/Engine/Source/Editor/UnrealEd/Private/Subsystems/AssetEditorSubsystem.cpp
--- a/Engine/Source/Editor/UnrealEd/Private/Subsystems/AssetEditorSubsystem.cpp
+++ b/Engine/Source/Editor/UnrealEd/Private/Subsystems/AssetEditorSubsystem.cpp
@@ -20,6 +20,8 @@
 #include "PackageTools.h"
 #include "UObject/PackageReload.h"
 #include "Interfaces/IAnalyticsProvider.h"
+#include "Misc/FeedbackContext.h"
+#include "Misc/ConfigCacheIni.h"
 
 
 #define LOCTEXT_NAMESPACE "AssetEditorSubsystem"
diff --git a/Engine/Source/Editor/UnrealEd/Private/Toolkits/SGlobalOpenAssetDialog.cpp b/Engine/Source/Editor/UnrealEd/Private/Toolkits/SGlobalOpenAssetDialog.cpp
--- a/Engine/Source/Editor/UnrealEd/Private/Toolkits/SGlobalOpenAssetDialog.cpp
+++ b/Engine/Source/Editor/UnrealEd/Private/Toolkits/SGlobalOpenAssetDialog.cpp
@@ -11,6 +11,7 @@
 #include "ContentBrowserModule.h"
 
 #include "Subsystems/AssetEditorSubsystem.h"
+#include "Editor.h"
 
 #define LOCTEXT_NAMESPACE "SGlobalOpenAssetDialog"
 
diff --git a/Engine/Source/Editor/UnrealEd/Public/MiniCurveEditor.h b/Engine/Source/Editor/UnrealEd/Public/MiniCurveEditor.h
--- a/Engine/Source/Editor/UnrealEd/Public/MiniCurveEditor.h
+++ b/Engine/Source/Editor/UnrealEd/Public/MiniCurveEditor.h
@@ -6,6 +6,7 @@
 #include "Widgets/DeclarativeSyntaxSupport.h"
 #include "Widgets/SCompoundWidget.h"
 #include "Framework/Docking/TabManager.h"
+#include "Subsystems/AssetEditorSubsystem.h"
 
 
 class FCurveOwnerInterface;
diff --git a/Engine/Source/Editor/UnrealEd/Public/Toolkits/AssetEditorManager.h b/Engine/Source/Editor/UnrealEd/Public/Toolkits/AssetEditorManager.h
--- a/Engine/Source/Editor/UnrealEd/Public/Toolkits/AssetEditorManager.h
+++ b/Engine/Source/Editor/UnrealEd/Public/Toolkits/AssetEditorManager.h
@@ -7,6 +7,7 @@
 #include "Containers/Ticker.h"
 #include "Framework/Docking/TabManager.h"
 #include "Toolkits/IToolkit.h"
+#include "Subsystems/AssetEditorSubsystem.h"
 
 class IMessageContext;
 class FMessageEndpoint;
