Added place-holder icon for the disaster recovery hub.

diff --git a/Engine/Plugins/Developer/Concert/ConcertSync/ConcertSyncClient/Source/ConcertSyncUI/Private/ConcertFrontendStyle.cpp b/Engine/Plugins/Developer/Concert/ConcertSync/ConcertSyncClient/Source/ConcertSyncUI/Private/ConcertFrontendStyle.cpp
--- a/Engine/Plugins/Developer/Concert/ConcertSync/ConcertSyncClient/Source/ConcertSyncUI/Private/ConcertFrontendStyle.cpp
+++ b/Engine/Plugins/Developer/Concert/ConcertSync/ConcertSyncClient/Source/ConcertSyncUI/Private/ConcertFrontendStyle.cpp
@@ -102,6 +102,9 @@ void FConcertFrontendStyle::Initialize()
 	StyleSet->Set("Concert.Join.Small",   new IMAGE_PLUGIN_BRUSH("Icons/icon_MultiUserMenuJoin_40x",   Icon20x20));
 	StyleSet->Set("Concert.Cancel.Small", new IMAGE_PLUGIN_BRUSH("Icons/icon_MultiUserMenuCancel_40x", Icon20x20));
 
+	// Disaster Recovery
+	StyleSet->Set("Concert.RecoveryHub", new IMAGE_PLUGIN_BRUSH("Icons/icon_RecoveryHub_32x", Icon16x16));
+
 	// Activity Text
 	{
 		FTextBlockStyle BoldText = FEditorStyle::Get().GetWidgetStyle<FTextBlockStyle>("RichTextBlock.Bold");
