Improved the customizability and usability of SConcertSessionRecovery and SConcertSessionActivities widgets.
  - Exposed more control over the visiblity/effects of the 'RecoverThrougRecoverAll/Cancel' buttons.
  - Exposed an attribute to explain why no activites would be displayed in the SConcertSessionActivities view.
  - Exposed few API useful to control the view.

diff --git a/Engine/Plugins/Developer/Concert/ConcertSync/ConcertSyncClient/Source/ConcertSyncUI/Private/SConcertSessionActivities.cpp b/Engine/Plugins/Developer/Concert/ConcertSync/ConcertSyncClient/Source/ConcertSyncUI/Private/SConcertSessionActivities.cpp
--- a/Engine/Plugins/Developer/Concert/ConcertSync/ConcertSyncClient/Source/ConcertSyncUI/Private/SConcertSessionActivities.cpp
+++ b/Engine/Plugins/Developer/Concert/ConcertSync/ConcertSyncClient/Source/ConcertSyncUI/Private/SConcertSessionActivities.cpp
@@ -399,7 +399,7 @@ void SConcertSessionActivities::Construct(const FArguments& InArgs)
 	{
 		HeaderRow->AddColumn(SHeaderRow::Column(ConcertSessionActivityUtils::PackageColumnId)
 			.DefaultLabel(LOCTEXT("Package", "Package"))
-			.ManualWidth(160));
+			.ManualWidth(200));
 	}
 
 	HeaderRow->AddColumn(SHeaderRow::Column(ConcertSessionActivityUtils::SummaryColumnId)
@@ -414,19 +414,34 @@ void SConcertSessionActivities::Construct(const FArguments& InArgs)
 		+SSplitter::Slot()
 		.Value(0.75)
 		[
-			SNew(SBorder)
-			.BorderImage(FEditorStyle::GetBrush("ToolPanel.GroupBorder"))
-			.BorderBackgroundColor(FSlateColor(FLinearColor(0.6, 0.6, 0.6)))
-			.Padding(0)
+			SNew(SOverlay)
+
+			+SOverlay::Slot() // Activity list itself.
+			[
+				SNew(SBorder)
+				.BorderImage(FEditorStyle::GetBrush("ToolPanel.GroupBorder"))
+				.BorderBackgroundColor(FSlateColor(FLinearColor(0.6, 0.6, 0.6)))
+				.Padding(0)
+				[
+					SAssignNew(ActivityView, SListView<TSharedPtr<FConcertClientSessionActivity>>)
+					.ListItemsSource(&Activities)
+					.OnGenerateRow(this, &SConcertSessionActivities::OnGenerateActivityRowWidget)
+					.SelectionMode(ESelectionMode::Single)
+					.AllowOverscroll(EAllowOverscroll::No)
+					.OnListViewScrolled(this, &SConcertSessionActivities::OnListViewScrolled)
+					.OnSelectionChanged(this, &SConcertSessionActivities::OnListViewSelectionChanged)
+					.HeaderRow(HeaderRow)
+				]
+			]
+
+			+SOverlay::Slot() // Display a reason why no activities are shown.
+			.HAlign(HAlign_Center)
+			.VAlign(VAlign_Center)
 			[
-				SAssignNew(ActivityView, SListView<TSharedPtr<FConcertClientSessionActivity>>)
-				.ListItemsSource(&Activities)
-				.OnGenerateRow(this, &SConcertSessionActivities::OnGenerateActivityRowWidget)
-				.SelectionMode(ESelectionMode::Single)
-				.AllowOverscroll(EAllowOverscroll::No)
-				.OnListViewScrolled(this, &SConcertSessionActivities::OnListViewScrolled)
-				.OnSelectionChanged(this, &SConcertSessionActivities::OnListViewSelectionChanged)
-				.HeaderRow(HeaderRow)
+				SNew(STextBlock)
+				.Visibility_Lambda([TextAttr = InArgs._NoActivitiesReasonText](){ return TextAttr.Get().IsEmptyOrWhitespace() ? EVisibility::Collapsed : EVisibility::Visible; })
+				.Text(InArgs._NoActivitiesReasonText)
+				.Justification(ETextJustify::Center)
 			]
 		]
 
diff --git a/Engine/Plugins/Developer/Concert/ConcertSync/ConcertSyncClient/Source/ConcertSyncUI/Private/SConcertSessionRecovery.cpp b/Engine/Plugins/Developer/Concert/ConcertSync/ConcertSyncClient/Source/ConcertSyncUI/Private/SConcertSessionRecovery.cpp
--- a/Engine/Plugins/Developer/Concert/ConcertSync/ConcertSyncClient/Source/ConcertSyncUI/Private/SConcertSessionRecovery.cpp
+++ b/Engine/Plugins/Developer/Concert/ConcertSync/ConcertSyncClient/Source/ConcertSyncUI/Private/SConcertSessionRecovery.cpp
@@ -22,6 +22,8 @@ void SConcertSessionRecovery::Construct(const FArguments& InArgs)
 	IntroductionText = InArgs._IntroductionText;
 	ParentWindow = InArgs._ParentWindow;
 	OnRestoreFn = InArgs._OnRestore;
+	OnCancelFn = InArgs._OnCancel;
+	IsRecoverThroughButtonVisible = InArgs._IsRecoverThroughButtonsVisible;
 
 	ActivityViewOptions = MakeShared<FConcertSessionActivitiesOptions>();
 	ActivityViewOptions->bEnableConnectionActivityFiltering = InArgs._IsConnectionActivityFilteringEnabled;
@@ -45,7 +47,8 @@ void SConcertSessionRecovery::Construct(const FArguments& InArgs)
 	.PackageActivitiesVisibility(ActivityViewOptions.Get(), &FConcertSessionActivitiesOptions::GetPackageActivitiesVisibility)
 	.TransactionActivitiesVisibility(ActivityViewOptions.Get(), &FConcertSessionActivitiesOptions::GetTransactionActivitiesVisibility)
 	.IgnoredActivitiesVisibility(ActivityViewOptions.Get(), &FConcertSessionActivitiesOptions::GetIgnoredActivitiesVisibility)
-	.DetailsAreaVisibility(InArgs._DetailsAreaVisibility);
+	.DetailsAreaVisibility(InArgs._DetailsAreaVisibility)
+	.NoActivitiesReasonText(InArgs._NoActivitiesReasonText);
 
 	ChildSlot
 	[
@@ -119,6 +122,7 @@ void SConcertSessionRecovery::Construct(const FArguments& InArgs)
 			[
 				SNew(SUniformGridPanel)
 				.SlotPadding(FMargin(2.0f, 0.0f))
+				.Visibility(InArgs._AreRecoverAllAndCancelButtonsVisible.Get() ? EVisibility::Visible : EVisibility::Collapsed)
 
 				+SUniformGridPanel::Slot(0, 0)
 				[
@@ -189,9 +193,24 @@ TSharedPtr<SWidget> SConcertSessionRecovery::MakeRecoverThroughWidget(TWeakPtr<F
 	return nullptr; // No overlay.
 }
 
+void SConcertSessionRecovery::Reset()
+{
+	ActivityView->Reset();
+}
+
+int32 SConcertSessionRecovery::GetTotalActivityNum() const
+{
+	return ActivityView->GetTotalActivityNum();
+}
+
+TSharedPtr<FConcertClientSessionActivity> SConcertSessionRecovery::GetMostRecentActivity() const
+{
+	return ActivityView->GetMostRecentActivity();
+}
+
 EVisibility SConcertSessionRecovery::GetRecoverThroughButtonVisibility(TSharedPtr<FConcertClientSessionActivity> Activity)
 {
-	return Activity == ActivityView->GetSelectedActivity() ? EVisibility::Visible : EVisibility::Hidden;
+	return Activity == ActivityView->GetSelectedActivity() && IsRecoverThroughButtonVisible.Get() ? EVisibility::Visible : EVisibility::Hidden;
 }
 
 FText SConcertSessionRecovery::GetRecoverThroughButtonTooltip() const
@@ -255,7 +274,14 @@ FReply SConcertSessionRecovery::OnCancelRecoveryClicked()
 {
 	check(!RecoveryThroughItem.IsValid());
 
-	DismissWindow();
+	if (OnCancelFn) // If a cancel behavior is bound.
+	{
+		OnCancelFn();
+	}
+	else // Default cancel behavior.
+	{
+		DismissWindow();
+	}
 	return FReply::Handled();
 }
 
diff --git a/Engine/Plugins/Developer/Concert/ConcertSync/ConcertSyncClient/Source/ConcertSyncUI/Public/SConcertSessionActivities.h b/Engine/Plugins/Developer/Concert/ConcertSync/ConcertSyncClient/Source/ConcertSyncUI/Public/SConcertSessionActivities.h
--- a/Engine/Plugins/Developer/Concert/ConcertSync/ConcertSyncClient/Source/ConcertSyncUI/Public/SConcertSessionActivities.h
+++ b/Engine/Plugins/Developer/Concert/ConcertSync/ConcertSyncClient/Source/ConcertSyncUI/Public/SConcertSessionActivities.h
@@ -130,6 +130,9 @@ public:
 
 		/** True to scroll the list down automatically (unless the user manually scrolled the list). */
 		SLATE_ARGUMENT(bool, IsAutoScrollEnabled)
+
+		/** Show/hide a message overlay above the activities list explaining why no activities are displayed. */
+		SLATE_ATTRIBUTE(FText, NoActivitiesReasonText)
 	SLATE_END_ARGS();
 
 	/**
diff --git a/Engine/Plugins/Developer/Concert/ConcertSync/ConcertSyncClient/Source/ConcertSyncUI/Public/SConcertSessionRecovery.h b/Engine/Plugins/Developer/Concert/ConcertSync/ConcertSyncClient/Source/ConcertSyncUI/Public/SConcertSessionRecovery.h
--- a/Engine/Plugins/Developer/Concert/ConcertSync/ConcertSyncClient/Source/ConcertSyncUI/Public/SConcertSessionRecovery.h
+++ b/Engine/Plugins/Developer/Concert/ConcertSync/ConcertSyncClient/Source/ConcertSyncUI/Public/SConcertSessionRecovery.h
@@ -35,7 +35,9 @@ public:
 		, _IsLockActivityFilteringEnabled(false)
 		, _IsPackageActivityFilteringEnabled(true)
 		, _IsTransactionActivityFilteringEnabled(true)
-		, _IsIgnoredActivityFilteringEnabled(false) { }
+		, _IsIgnoredActivityFilteringEnabled(false)
+		, _AreRecoverAllAndCancelButtonsVisible(true)
+		, _IsRecoverThroughButtonsVisible(true) {}
 
 		/** An introduction text to put the user in context. */
 		SLATE_ARGUMENT(FText, IntroductionText)
@@ -52,6 +54,9 @@ public:
 		/** Invoked when the user clicks the 'recover' button. */
 		SLATE_ARGUMENT(TFunction<bool(TSharedPtr<FConcertClientSessionActivity>)>, OnRestore)
 
+		/** Invoked when the user clicks the 'cancel' button. */
+		SLATE_ARGUMENT(TFunction<void()>, OnCancel)
+
 		/** Show/hide the column displaying the avatar color of the client who performed the activity. */
 		SLATE_ARGUMENT(EVisibility, ClientAvatarColorColumnVisibility)
 
@@ -81,6 +86,15 @@ public:
 
 		/** Show/hide the check box in the 'View Options' to filter unrecoverable activities. */
 		SLATE_ARGUMENT(bool, IsIgnoredActivityFilteringEnabled)
+
+		/** Show/hide the recover all and cancel buttons (so that if this widget is embedded in another one, the buttons can be replaced and put in a more appropriate location). */
+		SLATE_ATTRIBUTE(bool, AreRecoverAllAndCancelButtonsVisible)
+
+		/** Show/hide the recover though button. */
+		SLATE_ATTRIBUTE(bool, IsRecoverThroughButtonsVisible)
+
+		/** Show/hide a message overlay above the activities list explaining why no activities are displayed. */
+		SLATE_ATTRIBUTE(FText, NoActivitiesReasonText)
 	SLATE_END_ARGS();
 
 	/**
@@ -92,6 +106,18 @@ public:
 	/** Returns the activity, selected by the user, through which the session should be (or was) recovered or null to prevent recovery. */
 	TSharedPtr<FConcertClientSessionActivity> GetRecoverThroughItem() { return RecoveryThroughItem; }
 
+	/** Removes all activities from the view and reset the activity stream. */
+	void Reset();
+
+	/** Returns the total number of activities currently stored (no filter applied). */
+	int32 GetTotalActivityNum() const;
+
+	/** Returns the tooltip shown for the 'Recovery All' button. */
+	FText GetRecoverAllButtonTooltip() const;
+
+	/** Returns the most recent activity available, ignoring the current filter. */
+	TSharedPtr<FConcertClientSessionActivity> GetMostRecentActivity() const;
+
 private:
 	void OnSearchTextChanged(const FText& InFilterText);
 	void OnSearchTextCommitted(const FText& InFilterText, ETextCommit::Type CommitType);
@@ -102,7 +128,6 @@ private:
 
 	EVisibility GetRecoverThroughButtonVisibility(TSharedPtr<FConcertClientSessionActivity> Activity);
 	FText GetRecoverThroughButtonTooltip() const;
-	FText GetRecoverAllButtonTooltip() const;
 	TSharedPtr<SWidget> MakeRecoverThroughWidget(TWeakPtr<FConcertClientSessionActivity>, const FName&);
 	void RecoverThrough(TSharedPtr<FConcertClientSessionActivity> Item);
 
@@ -131,6 +156,12 @@ private:
 	/** The text displayed at the top to summarize the purpose of the window. */
 	FText IntroductionText;
 
+	/** Externally control the 'recovery through' button visibility. The internal logic to show it is run first, if it is still visible, this is evaluated.*/
+	TAttribute<bool> IsRecoverThroughButtonVisible;
+
 	/** The function invoked when the user clicks the restore button. Might not be bound. */
 	TFunction<bool(TSharedPtr<FConcertClientSessionActivity>)> OnRestoreFn;
+
+	/** The function invoked when the user clicks the cancel button. Might not be bound. */
+	TFunction<void()> OnCancelFn;
 };
