Added API to the ConcertSyncClientModule to track all Concert clients created.
  - Partial editgrated 8350450 from Dev-Enterprise.

diff --git a/Engine/Plugins/Developer/Concert/ConcertSync/ConcertSyncClient/Source/ConcertSyncClient/Private/ConcertSyncClientModule.cpp b/Engine/Plugins/Developer/Concert/ConcertSync/ConcertSyncClient/Source/ConcertSyncClient/Private/ConcertSyncClientModule.cpp
--- a/Engine/Plugins/Developer/Concert/ConcertSync/ConcertSyncClient/Source/ConcertSyncClient/Private/ConcertSyncClientModule.cpp
+++ b/Engine/Plugins/Developer/Concert/ConcertSync/ConcertSyncClient/Source/ConcertSyncClient/Private/ConcertSyncClientModule.cpp
@@ -85,7 +85,14 @@ public:
 
 	virtual TSharedRef<IConcertSyncClient> CreateClient(const FString& InRole) override
 	{
-		return MakeShared<FConcertSyncClient>(InRole, PackageBridge.Get(), TransactionBridge.Get());
+		// Remove dead clients.
+		Clients.RemoveAll([](TWeakPtr<IConcertSyncClient> WeakClient) { return !WeakClient.IsValid(); });
+
+		TSharedRef<IConcertSyncClient> NewClient = MakeShared<FConcertSyncClient>(InRole, PackageBridge.Get(), TransactionBridge.Get());
+		Clients.Add(NewClient);
+		OnClientCreated().Broadcast(NewClient);
+
+		return NewClient;
 	}
 
 	virtual IConcertClientPackageBridge& GetPackageBridge() override
@@ -98,9 +105,30 @@ public:
 		return *TransactionBridge;
 	}
 
+	virtual TArray<TSharedRef<IConcertSyncClient>> GetClients() const override
+	{
+		TArray<TSharedRef<IConcertSyncClient>> ActiveClients;
+		for (TWeakPtr<IConcertSyncClient> WeakClient : Clients)
+		{
+			if (TSharedPtr<IConcertSyncClient> Client = WeakClient.Pin())
+			{
+				ActiveClients.Add(Client.ToSharedRef());
+			}
+		}
+
+		return ActiveClients;
+	}
+
+	virtual FOnConcertClientCreated& OnClientCreated() override
+	{
+		return OnClientCreatedDelegate;
+	}
+
 private:
 	TUniquePtr<FConcertClientPackageBridge> PackageBridge;
 	TUniquePtr<FConcertClientTransactionBridge> TransactionBridge;
+	TArray<TWeakPtr<IConcertSyncClient>> Clients;
+	FOnConcertClientCreated OnClientCreatedDelegate;
 };
 
 IMPLEMENT_MODULE(FConcertSyncClientModule, ConcertSyncClient);
diff --git a/Engine/Plugins/Developer/Concert/ConcertSync/ConcertSyncClient/Source/ConcertSyncClient/Public/IConcertSyncClientModule.h b/Engine/Plugins/Developer/Concert/ConcertSync/ConcertSyncClient/Source/ConcertSyncClient/Public/IConcertSyncClientModule.h
--- a/Engine/Plugins/Developer/Concert/ConcertSync/ConcertSyncClient/Source/ConcertSyncClient/Public/IConcertSyncClientModule.h
+++ b/Engine/Plugins/Developer/Concert/ConcertSync/ConcertSyncClient/Source/ConcertSyncClient/Public/IConcertSyncClientModule.h
@@ -11,6 +11,8 @@ class IConcertClientPackageBridge;
 class IConcertClientTransactionBridge;
 class UConcertClientConfig;
 
+DECLARE_MULTICAST_DELEGATE_OneParam(FOnConcertClientCreated, TSharedRef<IConcertSyncClient>);
+
 /**
  * Interface for the Concert Sync Client module.
  */
@@ -63,4 +65,14 @@ public:
 	 * Get the bridge between the editor transaction system and Concert.
 	 */
 	virtual IConcertClientTransactionBridge& GetTransactionBridge() = 0;
+
+	/**
+	 * Returns the list of active clients.
+	 */
+	virtual TArray<TSharedRef<IConcertSyncClient>> GetClients() const = 0;
+
+	/**
+	 * Delegate invoked when a new concert sync client is created.
+	 */
+	virtual FOnConcertClientCreated& OnClientCreated() = 0;
 };
