Make it clearer that bIgnored is related to archive/restore

#rnx

diff --git a/Engine/Plugins/Developer/Concert/ConcertMain/Source/Concert/Private/ConcertModule.cpp b/Engine/Plugins/Developer/Concert/ConcertMain/Source/Concert/Private/ConcertModule.cpp
--- a/Engine/Plugins/Developer/Concert/ConcertMain/Source/Concert/Private/ConcertModule.cpp
+++ b/Engine/Plugins/Developer/Concert/ConcertMain/Source/Concert/Private/ConcertModule.cpp
@@ -23,9 +23,9 @@ public:
 		EndpointProvider.Reset();
 	}
 
-	virtual IConcertServerRef CreateServer(const FString& InRole, IConcertServerEventSink* InEventSink) override
+	virtual IConcertServerRef CreateServer(const FString& InRole, const FConcertSessionFilter& InAutoArchiveSessionFilter, IConcertServerEventSink* InEventSink) override
 	{
-		return MakeShared<FConcertServer, ESPMode::ThreadSafe>(InRole, InEventSink, EndpointProvider);
+		return MakeShared<FConcertServer, ESPMode::ThreadSafe>(InRole, InAutoArchiveSessionFilter, InEventSink, EndpointProvider);
 	}
 
 	virtual IConcertClientRef CreateClient(const FString& InRole) override
diff --git a/Engine/Plugins/Developer/Concert/ConcertMain/Source/Concert/Private/ConcertServer.cpp b/Engine/Plugins/Developer/Concert/ConcertMain/Source/Concert/Private/ConcertServer.cpp
--- a/Engine/Plugins/Developer/Concert/ConcertMain/Source/Concert/Private/ConcertServer.cpp
+++ b/Engine/Plugins/Developer/Concert/ConcertMain/Source/Concert/Private/ConcertServer.cpp
@@ -105,8 +105,9 @@ FConcertServerPaths::FConcertServerPaths(const FString& InRole, const FString& I
 {
 }
 
-FConcertServer::FConcertServer(const FString& InRole, IConcertServerEventSink* InEventSink, const TSharedPtr<IConcertEndpointProvider>& InEndpointProvider)
+FConcertServer::FConcertServer(const FString& InRole, const FConcertSessionFilter& InAutoArchiveSessionFilter, IConcertServerEventSink* InEventSink, const TSharedPtr<IConcertEndpointProvider>& InEndpointProvider)
 	: Role(InRole)
+	, AutoArchiveSessionFilter(InAutoArchiveSessionFilter)
 	, EventSink(InEventSink)
 	, EndpointProvider(InEndpointProvider)
 {
@@ -268,7 +269,7 @@ void FConcertServer::Shutdown()
 			bool bDeleteSessionData = true;
 			if (bAutoArchiveOnShutdown)
 			{
-				bDeleteSessionData = ArchiveLiveSession(LiveSessionId, FString(), FConcertSessionFilter()).IsValid();
+				bDeleteSessionData = ArchiveLiveSession(LiveSessionId, FString(), AutoArchiveSessionFilter).IsValid();
 			}
 			DestroyLiveSession(LiveSessionId, bDeleteSessionData);
 		}
@@ -484,7 +485,7 @@ void FConcertServer::ArchiveOfflineSessions()
 		ArchivedSessionInfo.SessionId = FGuid::NewGuid();
 		ArchivedSessionInfo.SessionName = ConcertServerUtil::GetArchiveName(LiveSessionInfo.SessionName, LiveSessionInfo.Settings);
 
-		if (EventSink->ArchiveSession(*this, Paths->GetSessionWorkingDir(LiveSessionInfo.SessionId), Paths->GetSessionSavedDir(ArchivedSessionInfo.SessionId), ArchivedSessionInfo, FConcertSessionFilter()))
+		if (EventSink->ArchiveSession(*this, Paths->GetSessionWorkingDir(LiveSessionInfo.SessionId), Paths->GetSessionSavedDir(ArchivedSessionInfo.SessionId), ArchivedSessionInfo, AutoArchiveSessionFilter))
 		{
 			UE_LOG(LogConcert, Display, TEXT("Deleting %s"), *Paths->GetSessionWorkingDir(LiveSessionInfo.SessionId));
 			ConcertUtil::DeleteDirectoryTree(*Paths->GetSessionWorkingDir(LiveSessionInfo.SessionId), *Paths->GetBaseWorkingDir());
diff --git a/Engine/Plugins/Developer/Concert/ConcertMain/Source/Concert/Private/ConcertServer.h b/Engine/Plugins/Developer/Concert/ConcertMain/Source/Concert/Private/ConcertServer.h
--- a/Engine/Plugins/Developer/Concert/ConcertMain/Source/Concert/Private/ConcertServer.h
+++ b/Engine/Plugins/Developer/Concert/ConcertMain/Source/Concert/Private/ConcertServer.h
@@ -77,7 +77,7 @@ private:
 class FConcertServer : public IConcertServer
 {
 public: 
-	FConcertServer(const FString& InRole, IConcertServerEventSink* InEventSink, const TSharedPtr<IConcertEndpointProvider>& InEndpointProvider);
+	FConcertServer(const FString& InRole, const FConcertSessionFilter& InAutoArchiveSessionFilter, IConcertServerEventSink* InEventSink, const TSharedPtr<IConcertEndpointProvider>& InEndpointProvider);
 	virtual ~FConcertServer();
 
 	virtual const FString& GetRole() const override;
@@ -190,6 +190,9 @@ private:
 	/** Cached root paths used by this server */
 	TUniquePtr<const FConcertServerPaths> Paths;
 
+	/** The session filter to apply when auto-archiving sessions */
+	FConcertSessionFilter AutoArchiveSessionFilter;
+
 	/** Sink functions for events that this server can emit */
 	IConcertServerEventSink* EventSink;
 
diff --git a/Engine/Plugins/Developer/Concert/ConcertMain/Source/Concert/Public/ConcertMessageData.h b/Engine/Plugins/Developer/Concert/ConcertMain/Source/Concert/Public/ConcertMessageData.h
--- a/Engine/Plugins/Developer/Concert/ConcertMain/Source/Concert/Public/ConcertMessageData.h
+++ b/Engine/Plugins/Developer/Concert/ConcertMain/Source/Concert/Public/ConcertMessageData.h
@@ -229,7 +229,7 @@ struct FConcertSessionFilter
 
 	/** True to include ignored activities */
 	UPROPERTY()
-	bool bIncludeIgnoredActivities = true;
+	bool bIncludeIgnoredActivities = false;
 };
 
 USTRUCT()
diff --git a/Engine/Plugins/Developer/Concert/ConcertMain/Source/Concert/Public/IConcertModule.h b/Engine/Plugins/Developer/Concert/ConcertMain/Source/Concert/Public/IConcertModule.h
--- a/Engine/Plugins/Developer/Concert/ConcertMain/Source/Concert/Public/IConcertModule.h
+++ b/Engine/Plugins/Developer/Concert/ConcertMain/Source/Concert/Public/IConcertModule.h
@@ -10,6 +10,8 @@ class IConcertServer;
 class IConcertClient;
 class IConcertServerEventSink;
 
+struct FConcertSessionFilter;
+
 typedef TSharedPtr<IConcertServer, ESPMode::ThreadSafe> IConcertServerPtr;
 typedef TSharedRef<IConcertServer, ESPMode::ThreadSafe> IConcertServerRef;
 typedef TSharedPtr<IConcertClient, ESPMode::ThreadSafe> IConcertClientPtr;
@@ -47,10 +49,11 @@ public:
 	/**
 	 * Create a server that will perform a certain role (eg, MultiUser, DisasterRecovery, etc)
 	 * @param InRole The role to create
+	 * @param InAutoArchiveSessionFilter The session filter to apply when auto-archiving sessions
 	 * @param InEventSink Sink functions for events that the server can emit
 	 * @return The server
 	 */
-	virtual IConcertServerRef CreateServer(const FString& InRole, IConcertServerEventSink* InEventSink) = 0;
+	virtual IConcertServerRef CreateServer(const FString& InRole, const FConcertSessionFilter& InAutoArchiveSessionFilter, IConcertServerEventSink* InEventSink) = 0;
 
 	/**
 	 * Create a client that will perform a certain role (eg, MultiUser, DisasterRecovery, etc)
diff --git a/Engine/Plugins/Developer/Concert/ConcertSync/ConcertSyncClient/Source/ConcertSyncClient/Private/ConcertClientWorkspace.cpp b/Engine/Plugins/Developer/Concert/ConcertSync/ConcertSyncClient/Source/ConcertSyncClient/Private/ConcertClientWorkspace.cpp
--- a/Engine/Plugins/Developer/Concert/ConcertSync/ConcertSyncClient/Source/ConcertSyncClient/Private/ConcertClientWorkspace.cpp
+++ b/Engine/Plugins/Developer/Concert/ConcertSync/ConcertSyncClient/Source/ConcertSyncClient/Private/ConcertClientWorkspace.cpp
@@ -707,16 +707,13 @@ void FConcertClientWorkspace::SetTransactionActivity(const FConcertSyncTransacti
 	if (LiveSession->GetSessionDatabase().SetTransactionActivity(InTransactionActivity))
 	{
 		PostActivityUpdated(InTransactionActivity);
-		if (!InTransactionActivity.bIgnored)
+		if (TransactionManager)
 		{
-			if (TransactionManager)
-			{
-				TransactionManager->HandleRemoteTransaction(InTransactionActivity.EndpointId, InTransactionActivity.EventId, bHasSyncedWorkspace);
-			}
-			if (LiveTransactionAuthors)
-			{
-				LiveTransactionAuthors->AddLiveTransactionActivity(InTransactionActivity.EndpointId, InTransactionActivity.EventData.Transaction.ModifiedPackages);
-			}
+			TransactionManager->HandleRemoteTransaction(InTransactionActivity.EndpointId, InTransactionActivity.EventId, bHasSyncedWorkspace);
+		}
+		if (LiveTransactionAuthors)
+		{
+			LiveTransactionAuthors->AddLiveTransactionActivity(InTransactionActivity.EndpointId, InTransactionActivity.EventData.Transaction.ModifiedPackages);
 		}
 	}
 	else
@@ -731,16 +728,13 @@ void FConcertClientWorkspace::SetPackageActivity(const FConcertSyncPackageActivi
 	if (LiveSession->GetSessionDatabase().SetPackageActivity(InPackageActivity))
 	{
 		PostActivityUpdated(InPackageActivity);
-		if (!InPackageActivity.bIgnored)
+		if (PackageManager)
 		{
-			if (PackageManager)
-			{
-				PackageManager->HandleRemotePackage(InPackageActivity.EndpointId, InPackageActivity.EventId, bHasSyncedWorkspace);
-			}
-			if (LiveTransactionAuthors)
-			{
-				LiveTransactionAuthors->ResolveLiveTransactionAuthorsForPackage(InPackageActivity.EventData.Package.Info.PackageName);
-			}
+			PackageManager->HandleRemotePackage(InPackageActivity.EndpointId, InPackageActivity.EventId, bHasSyncedWorkspace);
+		}
+		if (LiveTransactionAuthors)
+		{
+			LiveTransactionAuthors->ResolveLiveTransactionAuthorsForPackage(InPackageActivity.EventData.Package.Info.PackageName);
 		}
 	}
 	else
diff --git a/Engine/Plugins/Developer/Concert/ConcertSync/ConcertSyncCore/Source/ConcertSyncCore/Public/ConcertSyncSessionTypes.h b/Engine/Plugins/Developer/Concert/ConcertSync/ConcertSyncCore/Source/ConcertSyncCore/Public/ConcertSyncSessionTypes.h
--- a/Engine/Plugins/Developer/Concert/ConcertSync/ConcertSyncCore/Source/ConcertSyncCore/Public/ConcertSyncSessionTypes.h
+++ b/Engine/Plugins/Developer/Concert/ConcertSync/ConcertSyncCore/Source/ConcertSyncCore/Public/ConcertSyncSessionTypes.h
@@ -132,7 +132,7 @@ struct FConcertSyncActivity
 	UPROPERTY()
 	int64 ActivityId = 0;
 
-	/** True if this activity is included for tracking purposes only, and should be ignored when applying an activity feed */
+	/** True if this activity is included for tracking purposes only, and can be ignored when migrating a database */
 	UPROPERTY()
 	bool bIgnored = false;
 
diff --git a/Engine/Plugins/Developer/Concert/ConcertSync/ConcertSyncServer/Source/ConcertSyncServer/Private/ConcertSyncServer.cpp b/Engine/Plugins/Developer/Concert/ConcertSync/ConcertSyncServer/Source/ConcertSyncServer/Private/ConcertSyncServer.cpp
--- a/Engine/Plugins/Developer/Concert/ConcertSync/ConcertSyncServer/Source/ConcertSyncServer/Private/ConcertSyncServer.cpp
+++ b/Engine/Plugins/Developer/Concert/ConcertSync/ConcertSyncServer/Source/ConcertSyncServer/Private/ConcertSyncServer.cpp
@@ -213,8 +213,8 @@ bool MigrateSessionData(const FConcertSyncSessionDatabase& InSourceDatabase, con
 
 } // namespace ConcertSyncServerUtils
 
-FConcertSyncServer::FConcertSyncServer(const FString& InRole)
-	: ConcertServer(IConcertModule::Get().CreateServer(InRole, this))
+FConcertSyncServer::FConcertSyncServer(const FString& InRole, const FConcertSessionFilter& InAutoArchiveSessionFilter)
+	: ConcertServer(IConcertModule::Get().CreateServer(InRole, InAutoArchiveSessionFilter, this))
 	, SessionFlags(EConcertSyncSessionFlags::None)
 {
 }
diff --git a/Engine/Plugins/Developer/Concert/ConcertSync/ConcertSyncServer/Source/ConcertSyncServer/Private/ConcertSyncServer.h b/Engine/Plugins/Developer/Concert/ConcertSync/ConcertSyncServer/Source/ConcertSyncServer/Private/ConcertSyncServer.h
--- a/Engine/Plugins/Developer/Concert/ConcertSync/ConcertSyncServer/Source/ConcertSyncServer/Private/ConcertSyncServer.h
+++ b/Engine/Plugins/Developer/Concert/ConcertSync/ConcertSyncServer/Source/ConcertSyncServer/Private/ConcertSyncServer.h
@@ -12,13 +12,15 @@ class FConcertSyncServerLiveSession;
 class FConcertSyncServerArchivedSession;
 class FConcertSyncSessionDatabase;
 
+struct FConcertSessionFilter;
+
 /**
  * Implementation for a Concert Sync Server.
  */
 class FConcertSyncServer : public IConcertSyncServer, public IConcertServerEventSink
 {
 public:
-	explicit FConcertSyncServer(const FString& InRole);
+	FConcertSyncServer(const FString& InRole, const FConcertSessionFilter& InAutoArchiveSessionFilter);
 	virtual ~FConcertSyncServer();
 
 	//~ IConcertSyncServer interface
diff --git a/Engine/Plugins/Developer/Concert/ConcertSync/ConcertSyncServer/Source/ConcertSyncServer/Private/ConcertSyncServerModule.cpp b/Engine/Plugins/Developer/Concert/ConcertSync/ConcertSyncServer/Source/ConcertSyncServer/Private/ConcertSyncServerModule.cpp
--- a/Engine/Plugins/Developer/Concert/ConcertSync/ConcertSyncServer/Source/ConcertSyncServer/Private/ConcertSyncServerModule.cpp
+++ b/Engine/Plugins/Developer/Concert/ConcertSync/ConcertSyncServer/Source/ConcertSyncServer/Private/ConcertSyncServerModule.cpp
@@ -44,9 +44,9 @@ public:
 		return ServerConfig;
 	}
 
-	virtual TSharedRef<IConcertSyncServer> CreateServer(const FString& InRole) override
+	virtual TSharedRef<IConcertSyncServer> CreateServer(const FString& InRole, const FConcertSessionFilter& InAutoArchiveSessionFilter) override
 	{
-		return MakeShared<FConcertSyncServer>(InRole);
+		return MakeShared<FConcertSyncServer>(InRole, InAutoArchiveSessionFilter);
 	}
 };
 
diff --git a/Engine/Plugins/Developer/Concert/ConcertSync/ConcertSyncServer/Source/ConcertSyncServer/Public/ConcertSyncServerLoop.h b/Engine/Plugins/Developer/Concert/ConcertSync/ConcertSyncServer/Source/ConcertSyncServer/Public/ConcertSyncServerLoop.h
--- a/Engine/Plugins/Developer/Concert/ConcertSync/ConcertSyncServer/Source/ConcertSyncServer/Public/ConcertSyncServerLoop.h
+++ b/Engine/Plugins/Developer/Concert/ConcertSync/ConcertSyncServer/Source/ConcertSyncServer/Public/ConcertSyncServerLoop.h
@@ -3,6 +3,7 @@
 #pragma once
 
 #include "CoreMinimal.h"
+#include "ConcertMessageData.h"
 #include "ConcertSyncSessionFlags.h"
 
 class UConcertServerConfig;
@@ -21,6 +22,9 @@ struct FConcertSyncServerLoopInitArgs
 	/** Friendly name to use for this service (when showing it to a user in log messages, etc) */
 	FString ServiceFriendlyName;
 
+	/** The session filter to apply when auto-archiving sessions on this server */
+	FConcertSessionFilter ServiceAutoArchiveSessionFilter;
+
 	/** Function to get the server settings object to configure the server with with, or unbound to parse the default settings */
 	TFunction<const UConcertServerConfig*()> GetServerConfigFunc;
 };
diff --git a/Engine/Plugins/Developer/Concert/ConcertSync/ConcertSyncServer/Source/ConcertSyncServer/Public/ConcertSyncServerLoop.inl b/Engine/Plugins/Developer/Concert/ConcertSync/ConcertSyncServer/Source/ConcertSyncServer/Public/ConcertSyncServerLoop.inl
--- a/Engine/Plugins/Developer/Concert/ConcertSync/ConcertSyncServer/Source/ConcertSyncServer/Public/ConcertSyncServerLoop.inl
+++ b/Engine/Plugins/Developer/Concert/ConcertSync/ConcertSyncServer/Source/ConcertSyncServer/Public/ConcertSyncServerLoop.inl
@@ -81,7 +81,7 @@ int32 ConcertSyncServerLoop(int32 ArgC, TCHAR** ArgV, const FConcertSyncServerLo
 		}
 
 		// Setup Concert Sync to run in server mode.
-		TSharedPtr<IConcertSyncServer> ConcertSyncServer = IConcertSyncServerModule::Get().CreateServer(InitArgs.ServiceRole);
+		TSharedPtr<IConcertSyncServer> ConcertSyncServer = IConcertSyncServerModule::Get().CreateServer(InitArgs.ServiceRole, InitArgs.ServiceAutoArchiveSessionFilter);
 		ConcertSyncServer->Startup(ServerConfig, InitArgs.SessionFlags);
 
 		// if we have a default session, set it up properly
diff --git a/Engine/Plugins/Developer/Concert/ConcertSync/ConcertSyncServer/Source/ConcertSyncServer/Public/IConcertSyncServerModule.h b/Engine/Plugins/Developer/Concert/ConcertSync/ConcertSyncServer/Source/ConcertSyncServer/Public/IConcertSyncServerModule.h
--- a/Engine/Plugins/Developer/Concert/ConcertSync/ConcertSyncServer/Source/ConcertSyncServer/Public/IConcertSyncServerModule.h
+++ b/Engine/Plugins/Developer/Concert/ConcertSync/ConcertSyncServer/Source/ConcertSyncServer/Public/IConcertSyncServerModule.h
@@ -9,6 +9,8 @@
 class IConcertSyncServer;
 class UConcertServerConfig;
 
+struct FConcertSessionFilter;
+
 /**
  * Interface for the Concert Sync Server module.
  */
@@ -48,7 +50,8 @@ public:
 	/**
 	 * Create a sync server that will perform a certain role (eg, MultiUser, DisasterRecovery, etc)
 	 * @param InRole The role to create
+	 * @param InAutoArchiveSessionFilter The session filter to apply when auto-archiving sessions
 	 * @return The server
 	 */
-	virtual TSharedRef<IConcertSyncServer> CreateServer(const FString& InRole) = 0;
+	virtual TSharedRef<IConcertSyncServer> CreateServer(const FString& InRole, const FConcertSessionFilter& InAutoArchiveSessionFilter) = 0;
 };
diff --git a/Engine/Source/Programs/CrashReportClient/Private/RecoveryService.cpp b/Engine/Source/Programs/CrashReportClient/Private/RecoveryService.cpp
--- a/Engine/Source/Programs/CrashReportClient/Private/RecoveryService.cpp
+++ b/Engine/Source/Programs/CrashReportClient/Private/RecoveryService.cpp
@@ -76,8 +76,11 @@ bool FRecoveryService::Startup()
 	ServerConfig->WorkingDir = GetRecoveryServiceWorkingDir();
 	ServerConfig->ArchiveDir = GetRecoveryServiceArchivedDir();
 
+	FConcertSessionFilter AutoArchiveSessionFilter;
+	AutoArchiveSessionFilter.bIncludeIgnoredActivities = true;
+
 	// Start disaster recovery server.
-	Server = IConcertSyncServerModule::Get().CreateServer(TEXT("DisasterRecovery"));
+	Server = IConcertSyncServerModule::Get().CreateServer(TEXT("DisasterRecovery"), AutoArchiveSessionFilter);
 	Server->Startup(ServerConfig, EConcertSyncSessionFlags::Default_DisasterRecoverySession);
 
 	UE_LOG(CrashReportClientLog, Display, TEXT("%s Initialized (Name: %s, Version: %d.%d, Role: %s)"), RecoveryServiceName, *Server->GetConcertServer()->GetServerInfo().ServerName, ENGINE_MAJOR_VERSION, ENGINE_MINOR_VERSION, *Server->GetConcertServer()->GetRole());
diff --git a/Engine/Source/Programs/UnrealDisasterRecoveryService/Private/UnrealDisasterRecoveryService.cpp b/Engine/Source/Programs/UnrealDisasterRecoveryService/Private/UnrealDisasterRecoveryService.cpp
--- a/Engine/Source/Programs/UnrealDisasterRecoveryService/Private/UnrealDisasterRecoveryService.cpp
+++ b/Engine/Source/Programs/UnrealDisasterRecoveryService/Private/UnrealDisasterRecoveryService.cpp
@@ -22,6 +22,7 @@ INT32_MAIN_INT32_ARGC_TCHAR_ARGV()
 	ServerLoopInitArgs.SessionFlags = EConcertSyncSessionFlags::Default_DisasterRecoverySession;
 	ServerLoopInitArgs.ServiceRole = TEXT("DisasterRecovery");
 	ServerLoopInitArgs.ServiceFriendlyName = TEXT("Disaster Recovery Service");
+	ServerLoopInitArgs.ServiceAutoArchiveSessionFilter.bIncludeIgnoredActivities = true;
 
 	ServerLoopInitArgs.GetServerConfigFunc = [&EditorProcessId]() -> const UConcertServerConfig*
 	{
