Leave the package flags unmodified when saving a newly added asset for Multi-User purpose

diff --git a/Engine/Plugins/Developer/Concert/ConcertSync/ConcertSyncClient/Source/ConcertSyncClient/Private/ConcertClientPackageBridge.cpp b/Engine/Plugins/Developer/Concert/ConcertSync/ConcertSyncClient/Source/ConcertSyncClient/Private/ConcertClientPackageBridge.cpp
--- a/Engine/Plugins/Developer/Concert/ConcertSync/ConcertSyncClient/Source/ConcertSyncClient/Private/ConcertClientPackageBridge.cpp
+++ b/Engine/Plugins/Developer/Concert/ConcertSync/ConcertSyncClient/Source/ConcertSyncClient/Private/ConcertClientPackageBridge.cpp
@@ -227,8 +227,12 @@ void FConcertClientPackageBridge::HandleAssetAdded(UObject *Object)
 		UWorld* World = UWorld::FindWorldInPackage(Package);
 
 		const FString PackageFilename = FPaths::ProjectIntermediateDir() / TEXT("Concert") / TEXT("Temp") / FGuid::NewGuid().ToString() + (World ? FPackageName::GetMapPackageExtension() : FPackageName::GetAssetPackageExtension());
+		uint32 PackageFlags = Package->GetPackageFlags();
 		if (UPackage::SavePackage(Package, World, RF_Standalone, *PackageFilename, GWarn, nullptr, false, false, SAVE_NoError | SAVE_KeepDirty))
 		{
+			// Saving the newly added asset here shouldn't modify any of its package flags since it's a 'dummy' save i.e. PKG_NewlyCreated
+			Package->SetPackageFlagsTo(PackageFlags);
+
 			FConcertPackage Event;
 			ConcertClientPackageBridgeUtil::FillPackageInfo(Package, EConcertPackageUpdateType::Added, Event.Info);
 
