Avoid rogue slow task dialog appearing for a few frames due to concert checking connection.  Now we wait until the connection takes longer than a second before showing

diff --git a/Engine/Plugins/Developer/Concert/ConcertSync/ConcertSyncClient/Source/ConcertSyncClient/Private/ConcertClientWorkspace.cpp b/Engine/Plugins/Developer/Concert/ConcertSync/ConcertSyncClient/Source/ConcertSyncClient/Private/ConcertClientWorkspace.cpp
--- a/Engine/Plugins/Developer/Concert/ConcertSync/ConcertSyncClient/Source/ConcertSyncClient/Private/ConcertClientWorkspace.cpp
+++ b/Engine/Plugins/Developer/Concert/ConcertSync/ConcertSyncClient/Source/ConcertSyncClient/Private/ConcertClientWorkspace.cpp
@@ -499,7 +499,8 @@ void FConcertClientWorkspace::HandleConnectionChanged(IConcertClientSession& InS
 		bHasSyncedWorkspace = false;
 		bFinalizeWorkspaceSyncRequested = false;
 		InitialSyncSlowTask = MakeUnique<FScopedSlowTask>(1.0f, LOCTEXT("SynchronizingSession", "Synchronizing Session..."));
-		InitialSyncSlowTask->MakeDialog();
+		InitialSyncSlowTask->MakeDialogDelayed(1.0f);
+
 		FConcertSlowTaskStackWorkaround::Get().PushTask(InitialSyncSlowTask);
 
 		// Request our initial workspace sync for any new activity since we last joined
