Identify the Concert Package Bridge Handle Asset Added save  as a "from autosave" save for now to distinguish it from normal user save in the editor

diff --git a/Engine/Plugins/Developer/Concert/ConcertSync/ConcertSyncClient/Source/ConcertSyncClient/Private/ConcertClientPackageBridge.cpp b/Engine/Plugins/Developer/Concert/ConcertSync/ConcertSyncClient/Source/ConcertSyncClient/Private/ConcertClientPackageBridge.cpp
--- a/Engine/Plugins/Developer/Concert/ConcertSync/ConcertSyncClient/Source/ConcertSyncClient/Private/ConcertClientPackageBridge.cpp
+++ b/Engine/Plugins/Developer/Concert/ConcertSync/ConcertSyncClient/Source/ConcertSyncClient/Private/ConcertClientPackageBridge.cpp
@@ -227,7 +227,8 @@ void FConcertClientPackageBridge::HandleAssetAdded(UObject *Object)
 
 		const FString PackageFilename = FPaths::ProjectIntermediateDir() / TEXT("Concert") / TEXT("Temp") / FGuid::NewGuid().ToString() + (Asset && Asset->IsA<UWorld>() ? FPackageName::GetMapPackageExtension() : FPackageName::GetAssetPackageExtension());
 		uint32 PackageFlags = Package->GetPackageFlags();
-		if (UPackage::SavePackage(Package, World, RF_Standalone, *PackageFilename, GWarn, nullptr, false, false, SAVE_Async | SAVE_NoError | SAVE_KeepDirty))
+		// Identify this save as an autosave since we currently don't have a better alternative, but we want to distinguish it from a save triggered by the user and want to allow save callbacks to handle it as such
+		if (UPackage::SavePackage(Package, World, RF_Standalone, *PackageFilename, GWarn, nullptr, false, false, SAVE_Async | SAVE_NoError | SAVE_KeepDirty | SAVE_FromAutosave))
 		{
 			UPackage::WaitForAsyncFileWrites();
 			// Saving the newly added asset here shouldn't modify any of its package flags since it's a 'dummy' save i.e. PKG_NewlyCreated
