Add command line flag for enabling additional concert logging.

diff --git a/Engine/Plugins/Developer/Concert/ConcertSync/ConcertSyncClient/Source/ConcertSyncClient/Private/ConcertSyncClientModule.cpp b/Engine/Plugins/Developer/Concert/ConcertSync/ConcertSyncClient/Source/ConcertSyncClient/Private/ConcertSyncClientModule.cpp
--- a/Engine/Plugins/Developer/Concert/ConcertSync/ConcertSyncClient/Source/ConcertSyncClient/Private/ConcertSyncClientModule.cpp
+++ b/Engine/Plugins/Developer/Concert/ConcertSync/ConcertSyncClient/Source/ConcertSyncClient/Private/ConcertSyncClientModule.cpp
@@ -73,6 +73,9 @@ public:
 			FParse::Bool(CommandLine, TEXT("-CONCERTISHEADLESS="), ClientConfig->bIsHeadless);
 			ClientConfig->bIsHeadless |= !FSlateNotificationManager::Get().AreNotificationsAllowed();
 
+			ClientConfig->EndpointSettings.bEnableLogging |= FParse::Param(CommandLine, TEXT("CONCERTENABLELOGGING"));
+			FParse::Bool(CommandLine, TEXT("-CONCERTENABLELOGGING="), ClientConfig->EndpointSettings.bEnableLogging);
+
 			// CONCERTTAGS
 			{
 				FString CmdTags;
diff --git a/Engine/Plugins/Developer/Concert/ConcertSync/ConcertSyncServer/Source/ConcertSyncServer/Private/ConcertSyncServerModule.cpp b/Engine/Plugins/Developer/Concert/ConcertSync/ConcertSyncServer/Source/ConcertSyncServer/Private/ConcertSyncServerModule.cpp
--- a/Engine/Plugins/Developer/Concert/ConcertSync/ConcertSyncServer/Source/ConcertSyncServer/Private/ConcertSyncServerModule.cpp
+++ b/Engine/Plugins/Developer/Concert/ConcertSync/ConcertSyncServer/Source/ConcertSyncServer/Private/ConcertSyncServerModule.cpp
@@ -57,6 +57,9 @@ public:
 
 			ServerConfig->bCleanWorkingDir |= FParse::Param(CommandLine, TEXT("CONCERTCLEAN"));
 			FParse::Bool(CommandLine, TEXT("-CONCERTCLEAN="), ServerConfig->bCleanWorkingDir);
+
+			ServerConfig->EndpointSettings.bEnableLogging |= FParse::Param(CommandLine, TEXT("CONCERTENABLELOGGING"));
+			FParse::Bool(CommandLine, TEXT("-CONCERTENABLELOGGING="), ServerConfig->EndpointSettings.bEnableLogging);
 		}
 
 		return ServerConfig;
