- Adding livelink transaction filter for easier multi user setup

diff --git a/Engine/Plugins/Animation/LiveLink/LiveLink.uplugin b/Engine/Plugins/Animation/LiveLink/LiveLink.uplugin
--- a/Engine/Plugins/Animation/LiveLink/LiveLink.uplugin
+++ b/Engine/Plugins/Animation/LiveLink/LiveLink.uplugin
@@ -18,6 +18,11 @@
 		{
 			"Name": "Takes",
 			"Enabled": true
+		},
+		{
+			"Name": "ConcertSyncClient",
+			"Enabled": true,
+			"Optional": true
 		}
 	],
 	"Modules": [
@@ -50,6 +55,11 @@
 			"Name": "LiveLinkSequencer",
 			"Type": "UncookedOnly",
 			"LoadingPhase": "Default"
+		},
+		{
+			"Name": "LiveLinkMultiUser",
+			"Type": "UncookedOnly",
+			"LoadingPhase": "PostEngineInit"
 		}
 	]
 }
\ No newline at end of file
diff --git a/Engine/Plugins/Animation/LiveLink/Source/LiveLinkMultiUser/LiveLinkMultiUser.Build.cs b/Engine/Plugins/Animation/LiveLink/Source/LiveLinkMultiUser/LiveLinkMultiUser.Build.cs
new file mode 100644
--- /dev/null
+++ b/Engine/Plugins/Animation/LiveLink/Source/LiveLinkMultiUser/LiveLinkMultiUser.Build.cs
@@ -0,0 +1,24 @@
+// Copyright Epic Games, Inc. All Rights Reserved.
+
+namespace UnrealBuildTool.Rules
+{
+	public class LiveLinkMultiUser : ModuleRules
+	{
+		public LiveLinkMultiUser(ReadOnlyTargetRules Target) : base(Target)
+		{
+			PublicDependencyModuleNames.AddRange(
+			new string[]
+			{
+			});
+
+			PrivateDependencyModuleNames.AddRange(
+			new string[]
+			{
+				"ConcertSyncClient", // Dependancy to register our transaction filters. Needs to enable that plugin. Would be better to have modular feature interface functionality
+				"Core",
+				"CoreUObject",
+				"LiveLinkComponents",
+			});
+		}
+	}
+}
diff --git a/Engine/Plugins/Animation/LiveLink/Source/LiveLinkMultiUser/Private/LiveLinkMultiUserModule.cpp b/Engine/Plugins/Animation/LiveLink/Source/LiveLinkMultiUser/Private/LiveLinkMultiUserModule.cpp
new file mode 100644
--- /dev/null
+++ b/Engine/Plugins/Animation/LiveLink/Source/LiveLinkMultiUser/Private/LiveLinkMultiUserModule.cpp
@@ -0,0 +1,42 @@
+// Copyright Epic Games, Inc. All Rights Reserved.
+
+#include "LiveLinkMultiUserModule.h"
+
+#include "IConcertClientTransactionBridge.h"
+#include "IConcertSyncClientModule.h"
+#include "LiveLinkControllerBase.h"
+
+
+namespace LiveLinkMultiUserUtils
+{
+	ETransactionFilterResult HandleTransactionFiltering(UObject* ObjectToFilter, UPackage* ObjectsPackage)
+	{
+		//Always allow LiveLink controllers base class
+		if(Cast<ULiveLinkControllerBase>(ObjectToFilter) != nullptr)
+		{
+			return ETransactionFilterResult::IncludeObject;
+		}
+
+		return ETransactionFilterResult::UseDefault;
+	}
+}
+
+void FLiveLinkMultiUserModule::StartupModule()
+{
+	if (IConcertSyncClientModule::IsAvailable())
+	{
+		IConcertClientTransactionBridge& TransactionBridge = IConcertSyncClientModule::Get().GetTransactionBridge();
+		TransactionBridge.RegisterTransactionFilter( TEXT("LiveLinkTransactionFilter"), FTransactionFilterDelegate::CreateStatic(&LiveLinkMultiUserUtils::HandleTransactionFiltering));
+	}
+}
+
+void FLiveLinkMultiUserModule::ShutdownModule()
+{
+	if (IConcertSyncClientModule::IsAvailable())
+	{
+		IConcertClientTransactionBridge& TransactionBridge = IConcertSyncClientModule::Get().GetTransactionBridge();
+		TransactionBridge.UnregisterTransactionFilter( TEXT("LiveLinkTransactionFilter"));
+	}	
+}
+
+IMPLEMENT_MODULE(FLiveLinkMultiUserModule, LiveLinkMultiUser);
diff --git a/Engine/Plugins/Animation/LiveLink/Source/LiveLinkMultiUser/Private/LiveLinkMultiUserModule.h b/Engine/Plugins/Animation/LiveLink/Source/LiveLinkMultiUser/Private/LiveLinkMultiUserModule.h
new file mode 100644
--- /dev/null
+++ b/Engine/Plugins/Animation/LiveLink/Source/LiveLinkMultiUser/Private/LiveLinkMultiUserModule.h
@@ -0,0 +1,18 @@
+// Copyright Epic Games, Inc. All Rights Reserved.
+
+#pragma once
+
+#include "Modules/ModuleManager.h"
+
+/**
+ * Implements a simple LiveLink Multiuser module to register our own transaction filter.
+ */
+class FLiveLinkMultiUserModule : public IModuleInterface
+{
+public:
+
+	//~ Begin IModuleInterface interface
+	virtual void StartupModule() override;
+	virtual void ShutdownModule() override;
+	//~ End IModuleInterface interface
+};
diff --git a/Engine/Plugins/Developer/Concert/ConcertSync/ConcertSyncClient/Source/ConcertSyncClient/Private/ConcertAssetContainer.h b/Engine/Plugins/Developer/Concert/ConcertSync/ConcertSyncClient/Source/ConcertSyncClient/Private/ConcertAssetContainer.h
--- a/Engine/Plugins/Developer/Concert/ConcertSync/ConcertSyncClient/Source/ConcertSyncClient/Private/ConcertAssetContainer.h
+++ b/Engine/Plugins/Developer/Concert/ConcertSync/ConcertSyncClient/Source/ConcertSyncClient/Private/ConcertAssetContainer.h
@@ -30,48 +30,48 @@ public:
 	//
 
 	UPROPERTY(EditAnywhere, Category = Mesh)
-	UStaticMesh* GenericDesktopMesh;
+	TObjectPtr<UStaticMesh> GenericDesktopMesh;
 
 	UPROPERTY(EditAnywhere, Category = Mesh)
-	UStaticMesh* GenericHMDMesh;
+	TObjectPtr<UStaticMesh> GenericHMDMesh;
 
 	UPROPERTY(EditAnywhere, Category = Mesh)
-	UStaticMesh* VivePreControllerMesh;
+	TObjectPtr<UStaticMesh> VivePreControllerMesh;
 
 	UPROPERTY(EditAnywhere, Category = Mesh)
-	UStaticMesh* OculusControllerMesh;
+	TObjectPtr<UStaticMesh> OculusControllerMesh;
 
 	UPROPERTY(EditAnywhere, Category = Mesh)
-	UStaticMesh* GenericControllerMesh;
+	TObjectPtr<UStaticMesh> GenericControllerMesh;
 
 	UPROPERTY(EditAnywhere, Category = Mesh)
-	UStaticMesh* LaserPointerMesh;
+	TObjectPtr<UStaticMesh> LaserPointerMesh;
 
 	UPROPERTY(EditAnywhere, Category = Mesh)
-	UStaticMesh* LaserPointerEndMesh;
+	TObjectPtr<UStaticMesh> LaserPointerEndMesh;
 
 	UPROPERTY(EditAnywhere, Category = Mesh)
-	UStaticMesh* LaserPointerStartMesh;
+	TObjectPtr<UStaticMesh> LaserPointerStartMesh;
 
 	//
 	// Materials
 	//
 
 	UPROPERTY(EditAnywhere, Category = Material)
-	UMaterialInterface* PresenceMaterial;
+	TObjectPtr<UMaterialInterface> PresenceMaterial;
 
 	UPROPERTY(EditAnywhere, Category = Material)
-	UMaterialInterface* TextMaterial;
+	TObjectPtr<UMaterialInterface> TextMaterial;
 
 	UPROPERTY(EditAnywhere, Category = Desktop)
-	UMaterialInterface* HeadMaterial;
+	TObjectPtr<UMaterialInterface> HeadMaterial;
 
 	UPROPERTY(EditAnywhere, Category = Material)
-	UMaterialInterface* LaserCoreMaterial;
+	TObjectPtr<UMaterialInterface> LaserCoreMaterial;
 
 	UPROPERTY(EditAnywhere, Category = Material)
-	UMaterialInterface* LaserMaterial;
+	TObjectPtr<UMaterialInterface> LaserMaterial;
 
 	UPROPERTY(EditAnywhere, Category = Material)
-	UMaterialInterface* PresenceFadeMaterial;
+	TObjectPtr<UMaterialInterface> PresenceFadeMaterial;
 };
diff --git a/Engine/Plugins/Developer/Concert/ConcertSync/ConcertSyncClient/Source/ConcertSyncClient/Private/ConcertClientTransactionBridge.cpp b/Engine/Plugins/Developer/Concert/ConcertSync/ConcertSyncClient/Source/ConcertSyncClient/Private/ConcertClientTransactionBridge.cpp
--- a/Engine/Plugins/Developer/Concert/ConcertSync/ConcertSyncClient/Source/ConcertSyncClient/Private/ConcertClientTransactionBridge.cpp
+++ b/Engine/Plugins/Developer/Concert/ConcertSync/ConcertSyncClient/Source/ConcertSyncClient/Private/ConcertClientTransactionBridge.cpp
@@ -547,6 +547,8 @@ void FConcertClientTransactionBridge::HandleTransactionStateChanged(const FTrans
 
 void FConcertClientTransactionBridge::HandleObjectTransacted(UObject* InObject, const FTransactionObjectEvent& InTransactionEvent)
 {
+	SCOPED_CONCERT_TRACE(FConcertClientTransactionBridge_HandleObjectTransacted);
+
 	if (bIgnoreLocalTransactions)
 	{
 		return;
@@ -773,6 +775,8 @@ void FConcertClientTransactionBridge::OnEngineInitComplete()
 
 void FConcertClientTransactionBridge::OnEndFrame()
 {
+	SCOPED_CONCERT_TRACE(FConcertClientTransactionBridge_OnEndFrame);
+
 	for (auto OngoingTransactionsOrderIter = OngoingTransactionsOrder.CreateIterator(); OngoingTransactionsOrderIter; ++OngoingTransactionsOrderIter)
 	{
 		FOngoingTransaction* OngoingTransactionPtr = OngoingTransactions.Find(*OngoingTransactionsOrderIter);
diff --git a/Engine/Plugins/Developer/Concert/ConcertSync/ConcertSyncClient/Source/ConcertSyncClient/Public/ConcertClientDesktopPresenceActor.h b/Engine/Plugins/Developer/Concert/ConcertSync/ConcertSyncClient/Source/ConcertSyncClient/Public/ConcertClientDesktopPresenceActor.h
--- a/Engine/Plugins/Developer/Concert/ConcertSync/ConcertSyncClient/Source/ConcertSyncClient/Public/ConcertClientDesktopPresenceActor.h
+++ b/Engine/Plugins/Developer/Concert/ConcertSync/ConcertSyncClient/Source/ConcertSyncClient/Public/ConcertClientDesktopPresenceActor.h
@@ -9,7 +9,9 @@
 
 
 struct FDesktopPresenceUpdateEvent;
-
+class UMaterialInstanceDynamic;
+class USplineMeshComponent;
+class UStaticMeshComponent;
 
 /**
   * A AConcertClientDesktopPresenceActor is a child of AConcertClientPresenceActor that is used to represent users in desktop 
@@ -43,19 +45,19 @@ private:
 
 	/** The camera mesh component to show visually where the camera is placed */
 	UPROPERTY(VisibleAnywhere, BlueprintReadWrite, Category = "Rendering", meta = (AllowPrivateAccess = "true"))
-	class UStaticMeshComponent* DesktopMeshComponent;
+	TObjectPtr<UStaticMeshComponent> DesktopMeshComponent;
 
 	/** Spline mesh representing laser */
 	UPROPERTY(VisibleAnywhere, BlueprintReadWrite, Category = "Rendering", meta = (AllowPrivateAccess = "true"))
-	class USplineMeshComponent* LaserPointer;
+	TObjectPtr<USplineMeshComponent> LaserPointer;
 
 	/** Dynamic material for the laser */
 	UPROPERTY()
-	class UMaterialInstanceDynamic* LaserMid;
+	TObjectPtr<UMaterialInstanceDynamic> LaserMid;
 
 	/** Dynamic material for the laser */
 	UPROPERTY()
-	class UMaterialInstanceDynamic* LaserCoreMid;
+	TObjectPtr<UMaterialInstanceDynamic> LaserCoreMid;
 
 	UPROPERTY()
 	bool bMovingCamera;
diff --git a/Engine/Plugins/Developer/Concert/ConcertSync/ConcertSyncClient/Source/ConcertSyncClient/Public/ConcertClientPresenceActor.h b/Engine/Plugins/Developer/Concert/ConcertSync/ConcertSyncClient/Source/ConcertSyncClient/Public/ConcertClientPresenceActor.h
--- a/Engine/Plugins/Developer/Concert/ConcertSync/ConcertSyncClient/Source/ConcertSyncClient/Public/ConcertClientPresenceActor.h
+++ b/Engine/Plugins/Developer/Concert/ConcertSync/ConcertSyncClient/Source/ConcertSyncClient/Public/ConcertClientPresenceActor.h
@@ -11,6 +11,9 @@
 class FStructOnScope;
 class IConcertClientSession;
 class FConcertClientPresenceMode;
+class UStaticMeshComponent;
+class UTextRenderComponent;
+class UMaterialInstanceDynamic;
 
 /**
  * A ConcertClientPresenceActor is a transient actor representing other client presences during a concert client session.
@@ -76,19 +79,19 @@ protected:
 
 	/** The camera mesh component to show visually where the camera is placed */
 	UPROPERTY(VisibleAnywhere, BlueprintReadWrite, Category = "Rendering", meta = (AllowPrivateAccess = "true"))
-	class UStaticMeshComponent* PresenceMeshComponent;
+	TObjectPtr<UStaticMeshComponent> PresenceMeshComponent;
 	
 	/** The text render component to display the associated client's name */
 	UPROPERTY(VisibleAnywhere, BlueprintReadWrite, Category = "Rendering", meta = (AllowPrivateAccess = "true"))
-	class UTextRenderComponent* PresenceTextComponent;
+	TObjectPtr<UTextRenderComponent> PresenceTextComponent;
 
 	/** Dynamic material for the presence actor */
 	UPROPERTY()
-	class UMaterialInstanceDynamic* PresenceMID;
+	TObjectPtr<UMaterialInstanceDynamic> PresenceMID;
 
 	/** Dynamic material for the presence text */
 	UPROPERTY()
-	class UMaterialInstanceDynamic* TextMID;
+	TObjectPtr<UMaterialInstanceDynamic> TextMID;
 
 	/** Movement object to interpolate presence movement between update events.*/
 	TOptional<FConcertClientMovement> PresenceMovement;
diff --git a/Engine/Plugins/Developer/Concert/ConcertSync/ConcertSyncClient/Source/ConcertSyncClient/Public/ConcertClientVRPresenceActor.h b/Engine/Plugins/Developer/Concert/ConcertSync/ConcertSyncClient/Source/ConcertSyncClient/Public/ConcertClientVRPresenceActor.h
--- a/Engine/Plugins/Developer/Concert/ConcertSync/ConcertSyncClient/Source/ConcertSyncClient/Public/ConcertClientVRPresenceActor.h
+++ b/Engine/Plugins/Developer/Concert/ConcertSync/ConcertSyncClient/Source/ConcertSyncClient/Public/ConcertClientVRPresenceActor.h
@@ -8,6 +8,10 @@
 #include "Engine/Scene.h"
 #include "ConcertClientVRPresenceActor.generated.h"
 
+class UMaterialInstanceDynamic;
+class USplineComponent;
+class UStaticMeshComponent;
+
 /**
   * A ConcertClientVRPresenceActor is a child of AConcertClientPresenceActor that is used to represent users in VR 
   */
@@ -27,19 +31,19 @@ public:
 
 	/** The left controller mesh */
 	UPROPERTY()
-	class UStaticMeshComponent* LeftControllerMeshComponent;
+	TObjectPtr<UStaticMeshComponent> LeftControllerMeshComponent;
 
 	/** The right controller mesh */
 	UPROPERTY()
-	class UStaticMeshComponent* RightControllerMeshComponent;
+	TObjectPtr<UStaticMeshComponent> RightControllerMeshComponent;
 
 	/** Dynamic material for the laser */
 	UPROPERTY()
-	class UMaterialInstanceDynamic* LaserMid;
+	TObjectPtr<UMaterialInstanceDynamic> LaserMid;
 
 	/** Dynamic material for the laser */
 	UPROPERTY()
-	class UMaterialInstanceDynamic* LaserCoreMid;
+	TObjectPtr<UMaterialInstanceDynamic> LaserCoreMid;
 
 	UPROPERTY(BlueprintReadWrite, Category="Laser")
 	float LaserThickness = 0.5f;
@@ -62,19 +66,19 @@ private:
 
 	/** Spline for the left laser pointer, if any */
 	UPROPERTY()
-	class USplineComponent* LeftLaserSplineComponent;
+	TObjectPtr<USplineComponent> LeftLaserSplineComponent;
 
 	/** Spline meshes for the left curved laser, if any */
 	UPROPERTY()
-	TArray<class USplineMeshComponent*> LeftLaserSplineMeshComponents;
+	TArray<TObjectPtr<USplineMeshComponent>> LeftLaserSplineMeshComponents;
 
 	/** Spline for the right laser pointer, if any */
 	UPROPERTY()
-	class USplineComponent* RightLaserSplineComponent;
+	TObjectPtr<USplineComponent> RightLaserSplineComponent;
 
 	/** Spline meshes for the right curved laser, if any */
 	UPROPERTY()
-	TArray<class USplineMeshComponent*> RightLaserSplineMeshComponents;
+	TArray<TObjectPtr<USplineMeshComponent>> RightLaserSplineMeshComponents;
 
 	UPROPERTY()
 	bool bIsLeftControllerVisible;
diff --git a/Engine/Plugins/Developer/Concert/ConcertSync/ConcertSyncCore/Config/BaseConcertSyncCore.ini b/Engine/Plugins/Developer/Concert/ConcertSync/ConcertSyncCore/Config/BaseConcertSyncCore.ini
--- a/Engine/Plugins/Developer/Concert/ConcertSync/ConcertSyncCore/Config/BaseConcertSyncCore.ini
+++ b/Engine/Plugins/Developer/Concert/ConcertSync/ConcertSyncCore/Config/BaseConcertSyncCore.ini
@@ -9,6 +9,3 @@
 
 ; If you wish to use Composure in Multi-user, you may add that line under [/Script/ConcertSyncCore.ConcertSyncConfig] in your project engine config settings.
 ;+IncludeObjectClassFilters=(ObjectOuterClass=None,ObjectClasses=(Class'/Script/Composure.CompositingElementPass',Class'/Script/Composure.ComposurePostProcessPassPolicy'))
-
-; If you wish to use LiveLink in Multi-user, you may add that line under [/Script/ConcertSyncCore.ConcertSyncConfig] in your project engine config settings.
-;+IncludeObjectClassFilters=(ObjectOuterClass=Class'/Script/LiveLinkComponents.LiveLinkComponentController',ObjectClasses=(Class'/Script/LiveLinkComponents.LiveLinkControllerBase'))
