1. Move close events to a pending queue and process all on the end of frame.
2. Add a force close CVar for sequencer players.
3. Add a destroy / re-create phase when an asset is reloaded.  This allows
us to tear down any existing objects and reload it.  Note, because of a
stale handle we need to flush the GC at between destruction and
recreation.
4. Use the FSoftObjPath to resolve the name instead of the key. It is possible that the package name has changed.

diff --git a/Engine/Plugins/Developer/Concert/ConcertSync/ConcertSyncClient/Source/ConcertSyncClient/Private/ConcertClientSequencerManager.cpp b/Engine/Plugins/Developer/Concert/ConcertSync/ConcertSyncClient/Source/ConcertSyncClient/Private/ConcertClientSequencerManager.cpp
--- a/Engine/Plugins/Developer/Concert/ConcertSync/ConcertSyncClient/Source/ConcertSyncClient/Private/ConcertClientSequencerManager.cpp
+++ b/Engine/Plugins/Developer/Concert/ConcertSync/ConcertSyncClient/Source/ConcertSyncClient/Private/ConcertClientSequencerManager.cpp
@@ -13,12 +13,16 @@
 #include "IConcertSession.h"
 #include "ConcertSettings.h"
 
+#include "ConcertClientWorkspace.h"
+
 #include "Engine/GameEngine.h"
 #include "MovieSceneSequence.h"
 #include "MovieSceneTimeHelpers.h"
 #include "LevelSequencePlayer.h"
 #include "LevelSequenceActor.h"
 
+#include "UObject/PackageReload.h"
+
 #if WITH_EDITOR
 	#include "ISequencerModule.h"
 	#include "ISequencer.h"
@@ -42,6 +46,9 @@ static TAutoConsoleVariable<int32> CVarEnableRemoteSequencerOpen(TEXT("Concert.E
 // Enable opening Sequencer on remote machine whenever a sequencer is opened, if both instance have this option on.
 static TAutoConsoleVariable<int32> CVarEnableUnrelatedTimelineSync(TEXT("Concert.EnableUnrelatedTimelineSync"), 0, TEXT("Enable syncing unrelated sequencer timeline."));
 
+// Enable always closing player on remote machine whenever a sequencer is closed on an editor.
+static TAutoConsoleVariable<int32> CVarAlwaysCloseGamePlayerOnCloseEvent(
+	TEXT("Concert.AlwaysCloseGamePlayerOnCloseEvent"), 0, TEXT("Force this player to close even if other editors have it open. This CVar only works on `-game` instances."));
 
 FConcertClientSequencerManager::FConcertClientSequencerManager(IConcertSyncClient* InOwnerSyncClient)
 	: OwnerSyncClient(InOwnerSyncClient)
@@ -53,6 +60,7 @@ FConcertClientSequencerManager::FConcertClientSequencerManager(IConcertSyncClien
 	ISequencerModule& SequencerModule = FModuleManager::Get().LoadModuleChecked<ISequencerModule>("Sequencer");
 	OnSequencerCreatedHandle = SequencerModule.RegisterOnSequencerCreated(FOnSequencerCreated::FDelegate::CreateRaw(this, &FConcertClientSequencerManager::OnSequencerCreated));
 	FCoreDelegates::OnEndFrame.AddRaw(this, &FConcertClientSequencerManager::OnEndFrame);
+	FCoreUObjectDelegates::OnPackageReloaded.AddRaw(this, &FConcertClientSequencerManager::HandleAssetReload);
 }
 
 FConcertClientSequencerManager::~FConcertClientSequencerManager()
@@ -64,7 +72,7 @@ FConcertClientSequencerManager::~FConcertClientSequencerManager()
 	}
 
 	FCoreDelegates::OnEndFrame.RemoveAll(this);
-	
+
 	for (FOpenSequencerData& OpenSequencer : OpenSequencers)
 	{
 		TSharedPtr<ISequencer> Sequencer = OpenSequencer.WeakSequencer.Pin();
@@ -74,6 +82,7 @@ FConcertClientSequencerManager::~FConcertClientSequencerManager()
 			Sequencer->OnCloseEvent().Remove(OpenSequencer.OnCloseEventHandle);
 		}
 	}
+	FCoreUObjectDelegates::OnPackageReloaded.RemoveAll(this);
 }
 
 void FConcertClientSequencerManager::OnSequencerCreated(TSharedRef<ISequencer> InSequencer)
@@ -137,6 +146,11 @@ TArray<FConcertClientSequencerManager::FOpenSequencerData*, TInlineAllocator<1>>
 	return OutSequencers;
 }
 
+void FConcertClientSequencerManager::SetActiveWorkspace(TSharedPtr<FConcertClientWorkspace> InWorkspace)
+{
+	Workspace = InWorkspace;
+}
+
 float FConcertClientSequencerManager::GetLatencyCompensationMs() const
 {
 	IConcertClientRef ConcertClient = OwnerSyncClient->GetConcertClient();
@@ -201,6 +215,11 @@ void FConcertClientSequencerManager::SetSequencerRemoteOpen(bool bEnable)
 	CVarEnableRemoteSequencerOpen->AsVariable()->Set(bEnable ? 1 : 0);
 }
 
+bool FConcertClientSequencerManager::ShouldAlwaysCloseGameSequencerPlayer() const
+{
+	return CVarAlwaysCloseGamePlayerOnCloseEvent.GetValueOnAnyThread() > 0;
+}
+
 void FConcertClientSequencerManager::OnSequencerClosed(TSharedRef<ISequencer> InSequencer)
 {
 	// Find the associated open sequencer index
@@ -228,6 +247,7 @@ void FConcertClientSequencerManager::OnSequencerClosed(TSharedRef<ISequencer> In
 		UMovieSceneSequence* Sequence = InSequencer->GetRootMovieSceneSequence();
 		if (Sequence)
 		{
+			UE_LOG(LogConcertSequencerSync, Verbose, TEXT("Sending close event %s"), *Sequence->GetPathName());
 			FConcertSequencerCloseEvent CloseEvent;
 			CloseEvent.bMasterClose = ClosingSequencer.PlaybackMode == EPlaybackMode::Master; // this sequencer had control over the sequence playback
 			CloseEvent.SequenceObjectPath = *Sequence->GetPathName();
@@ -316,11 +336,51 @@ void FConcertClientSequencerManager::OnSequencerTimeChanged(TWeakPtr<ISequencer>
 
 void FConcertClientSequencerManager::OnCloseEvent(const FConcertSessionContext&, const FConcertSequencerCloseEvent& InEvent)
 {
-	FConcertSequencerState* SequencerState = SequencerStates.Find(*InEvent.SequenceObjectPath);
+	UE_LOG(LogConcertSequencerSync, Verbose, TEXT("OnCloseEvent: %s"), *InEvent.SequenceObjectPath);
+	PendingSequenceCloseEvents.Add(InEvent);
+}
+
+void FConcertClientSequencerManager::OnOpenEvent(const FConcertSessionContext&, const FConcertSequencerOpenEvent& InEvent)
+{
+	UE_LOG(LogConcertSequencerSync, Verbose, TEXT("OnOpenEvent: %s"), *InEvent.SequenceObjectPath);
+	PendingSequenceOpenEvents.Add(InEvent.SequenceObjectPath);
+}
+
+void FConcertClientSequencerManager::HandleAssetReload(const EPackageReloadPhase InPackageReloadPhase, FPackageReloadedEvent* InPackageReloadedEvent)
+{
+	if (InPackageReloadPhase != EPackageReloadPhase::PostPackageFixup)
+	{
+		return;
+	}
+
+	for (TTuple<FName, FSequencePlayer>& Item : SequencePlayers)
+	{
+		ALevelSequenceActor* LevelSequenceActor = Item.Value.Actor.Get();
+		// If we have a null LevelSequenceActor it means that it has already been destroyed by a close event.  We will not
+		// recreate the asset until all editors have closed it.
+		//
+		if (LevelSequenceActor)
+		{
+			const UPackage* PackageReloaded = InPackageReloadedEvent->GetNewPackage();
+			const FString AssetPathName = LevelSequenceActor->LevelSequence.GetAssetPathString();
+			FString PackageName = FPackageName::ObjectPathToPackageName(AssetPathName);
+			if (PackageName == PackageReloaded->GetFName().ToString())
+			{
+				UE_LOG(LogConcertSequencerSync, Display, TEXT("Rebuild required for LevelSequence %s"), *Item.Key.ToString());
+				TTuple<FName,FString> Pending(Item.Key,AssetPathName);
+				PendingDestroy.Emplace(MoveTemp(Pending));
+			}
+		}
+	}
+}
+
+void FConcertClientSequencerManager::ApplyTransportCloseEvent(const FConcertSequencerCloseEvent& PendingClose)
+{
+	FConcertSequencerState* SequencerState = SequencerStates.Find(*PendingClose.SequenceObjectPath);
 	if (SequencerState)
 	{
 		// if the event was that a sequencer that was in master playback mode was closed, stop playback
-		if (InEvent.bMasterClose)
+		if (PendingClose.bMasterClose)
 		{
 			SequencerState->PlayerStatus = EConcertMovieScenePlayerStatus::Stopped;
 			for (FOpenSequencerData* OpenSequencer : GatherRootSequencersByState(*SequencerState))
@@ -330,19 +390,13 @@ void FConcertClientSequencerManager::OnCloseEvent(const FConcertSessionContext&,
 			}
 		}
 		// otherwise, discard the state, it's no longer opened.
-		else 
+		else if (PendingClose.EditorsWithSequencerOpened == 0)
 		{
-			SequencerStates.Remove(*InEvent.SequenceObjectPath);
+			SequencerStates.Remove(*PendingClose.SequenceObjectPath);
 		}
 	}
 
-	ApplyCloseToPlayers(InEvent);
-}
-
-void FConcertClientSequencerManager::OnOpenEvent(const FConcertSessionContext&, const FConcertSequencerOpenEvent& InEvent)
-{
-	UE_LOG(LogConcertSequencerSync, Verbose, TEXT("OnOpenEvent: %s"), *InEvent.SequenceObjectPath);
-	PendingSequenceOpenEvents.Add(InEvent.SequenceObjectPath);
+	ApplyCloseToPlayers(PendingClose);
 }
 
 void FConcertClientSequencerManager::ApplyTransportOpenEvent(const FString& SequenceObjectPath)
@@ -354,30 +408,130 @@ void FConcertClientSequencerManager::ApplyTransportOpenEvent(const FString& Sequ
 		GEditor->GetEditorSubsystem<UAssetEditorSubsystem>()->OpenEditorForAsset(SequenceObjectPath);
 	}
 #endif
+
+	if (CVarEnableSequencePlayer.GetValueOnAnyThread() > 0)
+	{
+		CreateNewSequencePlayerIfNotExists(*SequenceObjectPath);
+	}
 }
 
-void FConcertClientSequencerManager::ApplyCloseToPlayers(const FConcertSequencerCloseEvent& InEvent)
+void FConcertClientSequencerManager::CreateNewSequencePlayerIfNotExists(const FString& SequenceObjectPath)
+{
+	// we do not have a player for this state yet
+	if (SequencePlayers.Contains(*SequenceObjectPath))
+	{
+		return;
+	}
+
+	UWorld* CurrentWorld = nullptr;
+	if (UGameEngine* GameEngine = Cast<UGameEngine>(GEngine))
+	{
+		CurrentWorld = GameEngine->GetGameWorld();
+	}
+	check(CurrentWorld);
+
+	ALevelSequenceActor* LevelSequenceActor = nullptr;
+
+	// Get the actual sequence
+	ULevelSequence* Sequence = LoadObject<ULevelSequence>( nullptr, *SequenceObjectPath);
+
+	if (!Sequence)
+	{
+		return;
+	}
+
+	FDelegateHandle Handle;
+	FMovieSceneSequencePlaybackSettings PlaybackSettings;
+	// Sequencer behaves differently to Player.
+	// Sequencer pauses at the last frame and Player Stops and goes to the first frame unless we set this flag.
+	PlaybackSettings.bPauseAtEnd = true;
+
+	// This call will initialeze LevelSequenceActor as an output parameter.
+	ULevelSequencePlayer* Player =ULevelSequencePlayer::CreateLevelSequencePlayer(CurrentWorld->PersistentLevel, Sequence, PlaybackSettings, LevelSequenceActor);
+	check(Player);
+
+	UMovieScene* Scene = Sequence->GetMovieScene();
+	check(Scene != nullptr);
+
+	Handle = Scene->OnSignatureChanged().AddLambda([SceneObj = TWeakObjectPtr<UMovieScene>(Scene),
+													LevelActorObj = TWeakObjectPtr<ALevelSequenceActor>(LevelSequenceActor)]() {
+		if (LevelActorObj.IsValid() && SceneObj.IsValid())
+		{
+			const TRange<FFrameNumber> PlayRange = SceneObj->GetPlaybackRange();
+			const FFrameRate TickResolution = SceneObj->GetTickResolution();
+			const FFrameRate DisplayRate = SceneObj->GetDisplayRate();
+
+			const FFrameNumber SrcStartFrame = UE::MovieScene::DiscreteInclusiveLower(PlayRange);
+			const FFrameNumber SrcEndFrame   = UE::MovieScene::DiscreteExclusiveUpper(PlayRange);
+
+			const FFrameTime EndingTime = ConvertFrameTime(SrcEndFrame, TickResolution, DisplayRate);
+
+			const FFrameNumber StartingFrame = ConvertFrameTime(SrcStartFrame, TickResolution, DisplayRate).FloorToFrame();
+			const FFrameNumber EndingFrame   = EndingTime.FloorToFrame();
+
+			const int32 CurrentDuration = LevelActorObj->GetSequencePlayer()->GetFrameDuration();
+			const FQualifiedFrameTime QualifiedStartTime = LevelActorObj->GetSequencePlayer()->GetStartTime();
+			const int32 NewDuration = (EndingFrame - StartingFrame).Value;
+			if (CurrentDuration != NewDuration || QualifiedStartTime.Time.GetFrame() != StartingFrame)
+			{
+				LevelActorObj->GetSequencePlayer()->SetFrameRange(StartingFrame.Value, NewDuration, EndingTime.GetSubFrame());
+			}
+		}
+	});
+	SequencePlayers.Add(*SequenceObjectPath, {LevelSequenceActor, MoveTemp(Handle)});
+}
+
+bool FConcertClientSequencerManager::CanClose(const FConcertSequencerCloseEvent& InEvent) const
+{
+	const bool bShouldClose = ShouldAlwaysCloseGameSequencerPlayer();
+
+	return InEvent.EditorsWithSequencerOpened == 0 || bShouldClose;
+}
+
+void FConcertClientSequencerManager::DestroyPlayer(FSequencePlayer& Player, ALevelSequenceActor* LevelSequenceActor)
 {
-	FSequencePlayer Player = SequencePlayers.FindRef(*InEvent.SequenceObjectPath);
-	ALevelSequenceActor* LevelSequenceActor = Player.Actor.Get();
 	if (LevelSequenceActor && LevelSequenceActor->SequencePlayer)
 	{
-		ULevelSequence *Sequence = LevelSequenceActor->GetSequence();
+		LevelSequenceActor->SequencePlayer->Stop();
+
+		ULevelSequence* Sequence = LevelSequenceActor->GetSequence();
 		if (Sequence)
 		{
-			UMovieScene *Scene = Sequence->GetMovieScene();
+			UMovieScene* Scene = Sequence->GetMovieScene();
 			if (Scene && Player.SignatureChangedHandle.IsValid())
 			{
 				Scene->OnSignatureChanged().Remove(Player.SignatureChangedHandle);
 			}
 		}
+		LevelSequenceActor->SetSequence(nullptr);
+		LevelSequenceActor->SequencePlayer = nullptr;
+		LevelSequenceActor->Destroy(false, false);
+	}
+}
+
+void FConcertClientSequencerManager::ApplyCloseToPlayers(const FConcertSequencerCloseEvent& InEvent)
+{
+	UE_LOG(LogConcertSequencerSync, Verbose, TEXT("CloseEvent: %s, is from master: %d"), *InEvent.SequenceObjectPath, InEvent.bMasterClose);
+	FSequencePlayer* Player = SequencePlayers.Find(*InEvent.SequenceObjectPath);
+	ALevelSequenceActor* LevelSequenceActor = Player ? Player->Actor.Get() : nullptr;
 
+	if (InEvent.bMasterClose && LevelSequenceActor && LevelSequenceActor->SequencePlayer)
+	{
 		LevelSequenceActor->SequencePlayer->Stop();
-		if (!InEvent.bMasterClose)
-		{
-			LevelSequenceActor->Destroy(false, false);
-			SequencePlayers.Remove(*InEvent.SequenceObjectPath);
-		}
+	}
+
+	if (!Player || !CanClose(InEvent))
+	{
+		return;
+	}
+
+	DestroyPlayer(*Player, LevelSequenceActor);
+
+	Player->Actor = nullptr;
+	if (InEvent.EditorsWithSequencerOpened == 0)
+	{
+		UE_LOG(LogConcertSequencerSync, Verbose, TEXT("Removing player: %s"), *InEvent.SequenceObjectPath);
+		SequencePlayers.Remove(*InEvent.SequenceObjectPath);
 	}
 }
 
@@ -393,6 +547,7 @@ void FConcertClientSequencerManager::ApplyTransportEvent(const FConcertSequencer
 		return;
 	}
 
+	UE_LOG(LogConcertSequencerSync, Verbose, TEXT("Apply Transport Event: %s"), *EventState.SequenceObjectPath);
 	TGuardValue<bool> ReentrancyGuard(bRespondingToTransportEvent, true);
 
 	// Update the sequencer pointing to the same sequence
@@ -503,63 +658,17 @@ void FConcertClientSequencerManager::ApplyEventToSequencers(const FConcertSequen
 
 void FConcertClientSequencerManager::ApplyEventToPlayers(const FConcertSequencerState& EventState)
 {
-	ULevelSequencePlayer* Player = nullptr;
-	ALevelSequenceActor* LevelSequenceActor = nullptr;
-	// we do not have a player for this state yet
-	if (!SequencePlayers.Contains(*EventState.SequenceObjectPath))
-	{
-		UWorld* CurrentWorld = nullptr;
-		if (UGameEngine* GameEngine = Cast<UGameEngine>(GEngine))
-		{
-			CurrentWorld = GameEngine->GetGameWorld();
-		}
-
-		// Get the actual sequence
-		ULevelSequence* Sequence = LoadObject<ULevelSequence>(nullptr, *EventState.SequenceObjectPath);
-		FDelegateHandle Handle;
-		if (Sequence && CurrentWorld)
-		{
-			FMovieSceneSequencePlaybackSettings PlaybackSettings;
-			// Sequencer behaves differently to Player. 
-			// Sequencer pauses at the last frame and Player Stops and goes to the first frame unless we set this flag.
-			PlaybackSettings.bPauseAtEnd = true;
-			Player = ULevelSequencePlayer::CreateLevelSequencePlayer(CurrentWorld->PersistentLevel, Sequence, PlaybackSettings, LevelSequenceActor);
-			UMovieScene* Scene = Sequence->GetMovieScene();
+	FSequencePlayer* SeqPlayer = SequencePlayers.Find(*EventState.SequenceObjectPath);
 
-			check(Scene != nullptr);
-			Handle = Scene->OnSignatureChanged().AddLambda([SceneObj = TWeakObjectPtr<UMovieScene>(Scene),
-															LevelActorObj = TWeakObjectPtr<ALevelSequenceActor>(LevelSequenceActor)]() {
-				if (LevelActorObj.IsValid() && SceneObj.IsValid())
-				{
-					const TRange<FFrameNumber> PlayRange = SceneObj->GetPlaybackRange();
-					const FFrameRate TickResolution = SceneObj->GetTickResolution();
-					const FFrameRate DisplayRate = SceneObj->GetDisplayRate();
-
-					const FFrameNumber SrcStartFrame = UE::MovieScene::DiscreteInclusiveLower(PlayRange);
-					const FFrameNumber SrcEndFrame   = UE::MovieScene::DiscreteExclusiveUpper(PlayRange);
-
-					const FFrameTime EndingTime = ConvertFrameTime(SrcEndFrame, TickResolution, DisplayRate);
-
-					const FFrameNumber StartingFrame = ConvertFrameTime(SrcStartFrame, TickResolution, DisplayRate).FloorToFrame();
-					const FFrameNumber EndingFrame   = EndingTime.FloorToFrame();
-
-					int32 CurrentDuration = LevelActorObj->GetSequencePlayer()->GetFrameDuration();
-					FQualifiedFrameTime QualifiedStartTime = LevelActorObj->GetSequencePlayer()->GetStartTime();
-					int32 NewDuration = (EndingFrame - StartingFrame).Value;
-					if (CurrentDuration != NewDuration || QualifiedStartTime.Time.GetFrame() != StartingFrame)
-					{
-						LevelActorObj->GetSequencePlayer()->SetFrameRange(StartingFrame.Value, NewDuration, EndingTime.GetSubFrame());
-		}
-				}
-			});
-		}
-		SequencePlayers.Add(*EventState.SequenceObjectPath, {LevelSequenceActor, MoveTemp(Handle)});
+	if (!SeqPlayer)
+	{
+		return;
 	}
 
-	LevelSequenceActor = SequencePlayers.FindChecked(*EventState.SequenceObjectPath).Actor.Get();
+	ALevelSequenceActor* LevelSequenceActor = SeqPlayer->Actor.Get();
 	if (LevelSequenceActor && LevelSequenceActor->SequencePlayer)
 	{
-		Player = LevelSequenceActor->SequencePlayer;
+		ULevelSequencePlayer* Player = LevelSequenceActor->SequencePlayer;
 		float LatencyCompensationMs = GetLatencyCompensationMs();
 
 		FFrameRate SequenceRate = Player->GetFrameRate();
@@ -643,6 +752,53 @@ void FConcertClientSequencerManager::ApplyEventToPlayers(const FConcertSequencer
 
 void FConcertClientSequencerManager::OnEndFrame()
 {
+	TSharedPtr<FConcertClientWorkspace> SharedWorkspace = Workspace.Pin();
+	if (SharedWorkspace)
+	{
+		if (!SharedWorkspace->CanProcessPendingPackages())
+		{
+			// There is currently a lock on the workspace.  We should wait for those to finish before processing
+			// sequencer events.
+			//
+			return;
+		}
+	}
+
+	for (const FConcertSequencerCloseEvent& CloseEvent: PendingSequenceCloseEvents)
+	{
+		ApplyTransportCloseEvent(CloseEvent);
+	}
+	PendingSequenceCloseEvents.Reset();
+
+	for (const TTuple<FName,FString>& Player: PendingDestroy)
+	{
+		if (FSequencePlayer* SequencePlayer = SequencePlayers.Find(Player.Get<0>()))
+		{
+			DestroyPlayer(*SequencePlayer, SequencePlayer->Actor.Get());
+			SequencePlayers.Remove(Player.Get<0>());
+			PendingCreate.Add(Player.Get<1>());
+		}
+	}
+
+	if (PendingDestroy.Num())
+	{
+		if (GEngine)
+		{
+			GEngine->ForceGarbageCollection(true);
+		}
+		PendingDestroy.Reset();
+		// Wait for the next frame to continue.  Garbage collection will not happen immediately so we need to wait for the
+		// next frame before we re-create.
+		//
+		return;
+	}
+
+	for (const FString& Player: PendingCreate)
+	{
+		CreateNewSequencePlayerIfNotExists(Player);
+	}
+	PendingCreate.Reset();
+
 	for (const FString& SequenceObjectPath : PendingSequenceOpenEvents)
 	{
 		ApplyTransportOpenEvent(SequenceObjectPath);
diff --git a/Engine/Plugins/Developer/Concert/ConcertSync/ConcertSyncClient/Source/ConcertSyncClient/Private/ConcertClientSequencerManager.h b/Engine/Plugins/Developer/Concert/ConcertSync/ConcertSyncClient/Source/ConcertSyncClient/Private/ConcertClientSequencerManager.h
--- a/Engine/Plugins/Developer/Concert/ConcertSync/ConcertSyncClient/Source/ConcertSyncClient/Private/ConcertClientSequencerManager.h
+++ b/Engine/Plugins/Developer/Concert/ConcertSync/ConcertSyncClient/Source/ConcertSyncClient/Private/ConcertClientSequencerManager.h
@@ -17,8 +17,15 @@ class ALevelSequenceActor;
 #if WITH_EDITOR
 
 class ISequencer;
+class FConcertClientWorkspace;
 struct FSequencerInitParams;
 
+struct FSequencePlayer
+{
+	TWeakObjectPtr<ALevelSequenceActor> Actor;
+	FDelegateHandle	SignatureChangedHandle;
+};
+
 /**
  * Sequencer manager that is held by the client sync module that keeps track of open sequencer UIs, regardless of whether a session is open or not
  * Events are registered to client sessions that will then operate on any tracked sequencer UIs
@@ -80,6 +87,13 @@ public:
 	 */
 	virtual bool IsSequencerRemoteOpenEnabled() const override;
 
+	/**
+	 * Checks the CVar to see if we are allowed to forcefully close the player on game instances.
+	 *
+	 * @return true if we should always close a sequence player on a -game instance.
+	 */
+	bool ShouldAlwaysCloseGameSequencerPlayer() const;
+
 	/**
 	 * Set the remote open option in Multi-User
 	 * which open sequencer for other user when this option is enabled on both user machines.
@@ -88,6 +102,9 @@ public:
 	 */
 	virtual void SetSequencerRemoteOpen(bool bEnable) override;
 
+	/** Assign the current active workspace to this sequencer. */
+	void SetActiveWorkspace(TSharedPtr<FConcertClientWorkspace> Workspace);
+
 private:
 	/** Enum signifying how a sequencer UI is currently playing. Necessary to prevent transport event contention. */
 	enum class EPlaybackMode
@@ -117,6 +134,8 @@ private:
 	};
 
 private:
+	void HandleAssetReload(const EPackageReloadPhase InPackageReloadPhase, FPackageReloadedEvent* InPackageReloadedEvent);
+
 	/**
 	 * Called when a sequencer closes.
 	 *
@@ -176,14 +195,21 @@ private:
 	void OnEndFrame();
 
 	/**
-	 * Apply a Sequencer open event 
+	 * Apply a Sequencer open event
 	 *
 	 * @param SequenceObjectPath	The sequence to open
 	 */
 	void ApplyTransportOpenEvent(const FString& SequenceObjectPath);
 
 	/**
-	 * Apply a Sequencer event 
+	 * Create a new sequence player to be used in -game instances.
+	 *
+	 * @param SequenceObjectPath	The object path for the Level Sequence object we are creating.
+	 */
+	void CreateNewSequencePlayerIfNotExists(const FString& SequenceObjectPath);
+
+	/**
+	 * Apply a Sequencer event
 	 *
 	 * @param PendingState	The pending state to apply
 	 */
@@ -210,6 +236,13 @@ private:
 	 */
 	void ApplyCloseToPlayers(const FConcertSequencerCloseEvent& InEvent);
 
+	/**
+	 * Apply CloseEvent to Open sequencers and players.
+	 *
+	 * @param PendingClose  The close event state.
+	 */
+	void ApplyTransportCloseEvent(const FConcertSequencerCloseEvent& PendingClose);
+
 	/**
 	 * Gather all the currently open sequencer UIs that have the specified path as their root sequence
 	 *
@@ -231,6 +264,13 @@ private:
 	}
 
 private:
+
+	/** Destroy the given sequence player with corresponding actor. */
+	void DestroyPlayer(FSequencePlayer& Player, ALevelSequenceActor* LevelSequenceActor);
+
+	/** Indicates if this sequence player can be closed. */
+	bool CanClose(const FConcertSequencerCloseEvent& InEvent) const;
+
 	/** Called by the MovieScene object when the property changes of the object. */
 	void OnPlayerSignatureChanged();
 
@@ -243,18 +283,25 @@ private:
 	/** List of pending sequencer open events to apply at end of frame. */
 	TArray<FString> PendingSequenceOpenEvents;
 
+	/** List of pending sequencer open events to apply at end of frame. */
+	TArray<FConcertSequencerCloseEvent> PendingSequenceCloseEvents;
+
+	/**
+	 * List of pending sequence players that are scheduled for destruction. The first member of the tuple is the key /
+	 * level sequence player we a removing.  The second member tuple is the soft object path name of the replacement
+	 * level sequence when it is recreated.
+	 */
+	TArray<TTuple<FName,FString>> PendingDestroy;
+
+	/** List of pending sequence players that are scheduled for creation. */
+	TArray<FString> PendingCreate;
+
 	/** Map of all currently opened Root Sequence State in a session, locally opened or not. */
 	TMap<FName, FConcertSequencerState> SequencerStates;
 
 	/** List of all locally opened sequencer. */
 	TArray<FOpenSequencerData> OpenSequencers;
 
-	struct FSequencePlayer
-	{
-		TWeakObjectPtr<ALevelSequenceActor> Actor;
-		FDelegateHandle						SignatureChangedHandle;
-	};
-
 	/** Map of opened sequence players, if not in editor mode. */
 	TMap<FName, FSequencePlayer> SequencePlayers;
 
@@ -266,6 +313,9 @@ private:
 	
 	/** Weak pointer to the client session with which to send events. May be null or stale. */
 	TWeakPtr<IConcertClientSession> WeakSession;
+
+	/** Weak pointer to the active workspace. */
+	TWeakPtr<FConcertClientWorkspace> Workspace;
 };
 
 #endif // WITH_EDITOR
diff --git a/Engine/Plugins/Developer/Concert/ConcertSync/ConcertSyncClient/Source/ConcertSyncClient/Private/ConcertSyncClient.cpp b/Engine/Plugins/Developer/Concert/ConcertSync/ConcertSyncClient/Source/ConcertSyncClient/Private/ConcertSyncClient.cpp
--- a/Engine/Plugins/Developer/Concert/ConcertSync/ConcertSyncClient/Source/ConcertSyncClient/Private/ConcertSyncClient.cpp
+++ b/Engine/Plugins/Developer/Concert/ConcertSync/ConcertSyncClient/Source/ConcertSyncClient/Private/ConcertSyncClient.cpp
@@ -142,6 +142,7 @@ void FConcertSyncClient::CreateWorkspace(const TSharedRef<FConcertSyncClientLive
 		// TODO: Revisit this, as all it seems to be used for now is forcing redirectors to be left behind
 		SourceControlProxy->SetWorkspace(Workspace);
 	}
+	SequencerManager->SetActiveWorkspace(Workspace);
 #endif
 }
 
diff --git a/Engine/Plugins/Developer/Concert/ConcertSync/ConcertSyncCore/Source/ConcertSyncCore/Public/ConcertSequencerMessages.h b/Engine/Plugins/Developer/Concert/ConcertSync/ConcertSyncCore/Source/ConcertSyncCore/Public/ConcertSequencerMessages.h
--- a/Engine/Plugins/Developer/Concert/ConcertSync/ConcertSyncCore/Source/ConcertSyncCore/Public/ConcertSequencerMessages.h
+++ b/Engine/Plugins/Developer/Concert/ConcertSync/ConcertSyncCore/Source/ConcertSyncCore/Public/ConcertSequencerMessages.h
@@ -79,6 +79,9 @@ struct FConcertSequencerCloseEvent
 
 	UPROPERTY()
 	bool bMasterClose = false;
+
+	UPROPERTY()
+	int32 EditorsWithSequencerOpened = -1;
 };
 
 /**
diff --git a/Engine/Plugins/Developer/Concert/ConcertSync/ConcertSyncServer/Source/ConcertSyncServer/Private/ConcertServerSequencerManager.cpp b/Engine/Plugins/Developer/Concert/ConcertSync/ConcertSyncServer/Source/ConcertSyncServer/Private/ConcertServerSequencerManager.cpp
--- a/Engine/Plugins/Developer/Concert/ConcertSync/ConcertSyncServer/Source/ConcertSyncServer/Private/ConcertServerSequencerManager.cpp
+++ b/Engine/Plugins/Developer/Concert/ConcertSync/ConcertSyncServer/Source/ConcertSyncServer/Private/ConcertServerSequencerManager.cpp
@@ -71,20 +71,17 @@ void FConcertServerSequencerManager::HandleSequencerCloseEvent(const FConcertSes
 	if (SequencerState)
 	{
 		SequencerState->ClientEndpointIds.Remove(InEventContext.SourceEndpointId);
-		if (SequencerState->ClientEndpointIds.Num() == 0)
+		// Forward a normal close event to clients
+		const int32 NumOpen = SequencerState->ClientEndpointIds.Num();
+		FConcertSequencerCloseEvent CloseEvent;
+		CloseEvent.bMasterClose = NumOpen != 0 && InEvent.bMasterClose;
+		CloseEvent.EditorsWithSequencerOpened = NumOpen;
+		CloseEvent.SequenceObjectPath = InEvent.SequenceObjectPath;
+		LiveSession->GetSession().SendCustomEvent(CloseEvent, LiveSession->GetSession().GetSessionClientEndpointIds(), EConcertMessageFlags::ReliableOrdered | EConcertMessageFlags::UniqueId);
+		if (NumOpen == 0)
 		{
-			// Forward a normal close event to clients
-			FConcertSequencerCloseEvent CloseEvent;
-			CloseEvent.bMasterClose = false;
-			CloseEvent.SequenceObjectPath = InEvent.SequenceObjectPath;
-			LiveSession->GetSession().SendCustomEvent(CloseEvent, LiveSession->GetSession().GetSessionClientEndpointIds(), EConcertMessageFlags::ReliableOrdered | EConcertMessageFlags::UniqueId);
 			SequencerStates.Remove(*InEvent.SequenceObjectPath);
 		}
-		// if a sequence was close while it was the master, forward it to client
-		else if (InEvent.bMasterClose)
-		{
-			LiveSession->GetSession().SendCustomEvent(InEvent, LiveSession->GetSession().GetSessionClientEndpointIds(), EConcertMessageFlags::ReliableOrdered | EConcertMessageFlags::UniqueId);
-		}
 	}
 }
 
@@ -97,13 +94,15 @@ void FConcertServerSequencerManager::HandleSessionClientChanged(IConcertServerSe
 		for (auto It = SequencerStates.CreateIterator(); It; ++It)
 		{
 			It->Value.ClientEndpointIds.Remove(InClientInfo.ClientEndpointId);
-			if (It->Value.ClientEndpointIds.Num() == 0)
-			{
-				// Forward the close event to clients
-				FConcertSequencerCloseEvent CloseEvent;
-				CloseEvent.SequenceObjectPath = It->Key.ToString();
-				LiveSession->GetSession().SendCustomEvent(CloseEvent, LiveSession->GetSession().GetSessionClientEndpointIds(), EConcertMessageFlags::ReliableOrdered|EConcertMessageFlags::UniqueId);
+			const int32 NumOpen = It->Value.ClientEndpointIds.Num();
+			// Forward the close event to clients
+			FConcertSequencerCloseEvent CloseEvent;
+			CloseEvent.EditorsWithSequencerOpened = NumOpen;
+			CloseEvent.SequenceObjectPath = It->Key.ToString();
+			LiveSession->GetSession().SendCustomEvent(CloseEvent, LiveSession->GetSession().GetSessionClientEndpointIds(), EConcertMessageFlags::ReliableOrdered|EConcertMessageFlags::UniqueId);
 
+			if (NumOpen == 0)
+			{
 				It.RemoveCurrent();
 			}
 		}
