Make sure we are not running in editor when trying to create a sequence player (which is only supported on -game nodes).

diff --git a/Engine/Plugins/Developer/Concert/ConcertSync/ConcertSyncClient/Source/ConcertSyncClient/Private/ConcertClientSequencerManager.cpp b/Engine/Plugins/Developer/Concert/ConcertSync/ConcertSyncClient/Source/ConcertSyncClient/Private/ConcertClientSequencerManager.cpp
--- a/Engine/Plugins/Developer/Concert/ConcertSync/ConcertSyncClient/Source/ConcertSyncClient/Private/ConcertClientSequencerManager.cpp
+++ b/Engine/Plugins/Developer/Concert/ConcertSync/ConcertSyncClient/Source/ConcertSyncClient/Private/ConcertClientSequencerManager.cpp
@@ -407,9 +407,9 @@ void FConcertClientSequencerManager::ApplyTransportOpenEvent(const FString& Sequ
 	{
 		GEditor->GetEditorSubsystem<UAssetEditorSubsystem>()->OpenEditorForAsset(SequenceObjectPath);
 	}
-#endif
 
-	if (CVarEnableSequencePlayer.GetValueOnAnyThread() > 0)
+#endif
+	if (!GIsEditor && CVarEnableSequencePlayer.GetValueOnAnyThread() > 0)
 	{
 		CreateNewSequencePlayerIfNotExists(*SequenceObjectPath);
 	}
