Make sure we check the non-sandbox area.  If it exists there, even if it was deleted, then it exists and the call should return true.

diff --git a/Engine/Plugins/Developer/Concert/ConcertSync/ConcertSyncClient/Source/ConcertSyncClient/Private/ConcertSandboxPlatformFile.cpp b/Engine/Plugins/Developer/Concert/ConcertSync/ConcertSyncClient/Source/ConcertSyncClient/Private/ConcertSandboxPlatformFile.cpp
--- a/Engine/Plugins/Developer/Concert/ConcertSync/ConcertSyncClient/Source/ConcertSyncClient/Private/ConcertSandboxPlatformFile.cpp
+++ b/Engine/Plugins/Developer/Concert/ConcertSync/ConcertSyncClient/Source/ConcertSyncClient/Private/ConcertSandboxPlatformFile.cpp
@@ -193,11 +193,6 @@ bool FConcertSandboxPlatformFile::FileExists(const TCHAR* Filename)
 
 	if (ResolvedPath.HasSandboxPath())
 	{
-		if (IsPathDeleted(ResolvedPath))
-		{
-			return false;
-		}
-		
 		if (LowerLevel->FileExists(*ResolvedPath.GetSandboxPath()))
 		{
 			return true;
@@ -213,11 +208,6 @@ int64 FConcertSandboxPlatformFile::FileSize(const TCHAR* Filename)
 
 	if (ResolvedPath.HasSandboxPath())
 	{
-		if (IsPathDeleted(ResolvedPath))
-		{
-			return -1;
-		}
-		
 		if (LowerLevel->FileExists(*ResolvedPath.GetSandboxPath()))
 		{
 			return LowerLevel->FileSize(*ResolvedPath.GetSandboxPath());
