Backout original fix to Platform file for deleted assets.

It caused problems with moving files within Multi-user session and need to develop a more extensive fix.

diff --git a/Engine/Plugins/Developer/Concert/ConcertSync/ConcertSyncClient/Source/ConcertSyncClient/Private/ConcertSandboxPlatformFile.cpp b/Engine/Plugins/Developer/Concert/ConcertSync/ConcertSyncClient/Source/ConcertSyncClient/Private/ConcertSandboxPlatformFile.cpp
--- a/Engine/Plugins/Developer/Concert/ConcertSync/ConcertSyncClient/Source/ConcertSyncClient/Private/ConcertSandboxPlatformFile.cpp
+++ b/Engine/Plugins/Developer/Concert/ConcertSync/ConcertSyncClient/Source/ConcertSyncClient/Private/ConcertSandboxPlatformFile.cpp
@@ -193,6 +193,11 @@ bool FConcertSandboxPlatformFile::FileExists(const TCHAR* Filename)
 
 	if (ResolvedPath.HasSandboxPath())
 	{
+		if (IsPathDeleted(ResolvedPath))
+		{
+			return false;
+		}
+
 		if (LowerLevel->FileExists(*ResolvedPath.GetSandboxPath()))
 		{
 			return true;
@@ -208,6 +213,11 @@ int64 FConcertSandboxPlatformFile::FileSize(const TCHAR* Filename)
 
 	if (ResolvedPath.HasSandboxPath())
 	{
+		if (IsPathDeleted(ResolvedPath))
+		{
+			return -1;
+		}
+
 		if (LowerLevel->FileExists(*ResolvedPath.GetSandboxPath()))
 		{
 			return LowerLevel->FileSize(*ResolvedPath.GetSandboxPath());
