Add a special check for persisting deleted assets.

diff --git a/Engine/Plugins/Developer/Concert/ConcertSync/ConcertSyncClient/Source/ConcertSyncClient/Private/ConcertClientPackageManager.cpp b/Engine/Plugins/Developer/Concert/ConcertSync/ConcertSyncClient/Source/ConcertSyncClient/Private/ConcertClientPackageManager.cpp
--- a/Engine/Plugins/Developer/Concert/ConcertSync/ConcertSyncClient/Source/ConcertSyncClient/Private/ConcertClientPackageManager.cpp
+++ b/Engine/Plugins/Developer/Concert/ConcertSync/ConcertSyncClient/Source/ConcertSyncClient/Private/ConcertClientPackageManager.cpp
@@ -288,6 +288,50 @@ bool FConcertClientPackageManager::HasSessionChanges() const
 	return bHasSessionChanges;
 }
 
+TOptional<FString> FConcertClientPackageManager::GetValidPackageSessionPath(FName PackageName) const
+{
+#if WITH_EDITOR
+	FString Filename;
+	if (FPackageName::DoesPackageExist(PackageName.ToString(), &Filename))
+	{
+		return Filename;
+	}
+	return GetDeletedPackagePath(PackageName);
+#else
+	return {};
+#endif
+}
+
+TOptional<FString> FConcertClientPackageManager::GetDeletedPackagePath(FName PackageName) const
+{
+#if WITH_EDITOR
+	check(SandboxPlatformFile.IsValid());
+	FConcertSandboxPlatformFile* PlatformFile = SandboxPlatformFile.Get();
+	auto ShouldPersistPackageWithExtension = [PlatformFile](const FString& PackageName, const FString& Extension) -> TOptional<FString>
+	{
+		FString FullPath;
+		if (FPackageName::TryConvertLongPackageNameToFilename(PackageName, FullPath, Extension))
+		{
+			if (PlatformFile->DeletedPackageExistsInNonSandbox(FullPath))
+			{
+				return MoveTemp(FullPath);
+			}
+		}
+		return {};
+	};
+	FString PackageNameAsString = PackageName.ToString();
+	if (TOptional<FString> AsMap = ShouldPersistPackageWithExtension(PackageNameAsString, FPackageName::GetMapPackageExtension()))
+	{
+		return AsMap;
+	}
+	if (TOptional<FString> AsAsset = ShouldPersistPackageWithExtension(PackageNameAsString, FPackageName::GetAssetPackageExtension()))
+	{
+		return AsAsset;
+	}
+#endif
+	return {};
+}
+
 bool FConcertClientPackageManager::PersistSessionChanges(TArrayView<const FName> InPackagesToPersist, ISourceControlProvider* SourceControlProvider, TArray<FText>* OutFailureReasons)
 {
 #if WITH_EDITOR
@@ -295,12 +339,11 @@ bool FConcertClientPackageManager::PersistSessionChanges(TArrayView<const FName>
 	{
 		// Transform all the package names into actual filenames
 		TArray<FString, TInlineAllocator<8>> FilesToPersist;
-		FString Filename;
 		for (const FName& PackageName : InPackagesToPersist)
 		{
-			if (FPackageName::DoesPackageExist(PackageName.ToString(), &Filename))
+			if (TOptional<FString> ValidPath = GetValidPackageSessionPath(PackageName))
 			{
-				FilesToPersist.Add(MoveTemp(Filename));
+				FilesToPersist.Add(MoveTemp(ValidPath.GetValue()));
 			}
 		}
 		return SandboxPlatformFile->PersistSandbox(FilesToPersist, SourceControlProvider, OutFailureReasons);
diff --git a/Engine/Plugins/Developer/Concert/ConcertSync/ConcertSyncClient/Source/ConcertSyncClient/Private/ConcertClientPackageManager.h b/Engine/Plugins/Developer/Concert/ConcertSync/ConcertSyncClient/Source/ConcertSyncClient/Private/ConcertClientPackageManager.h
--- a/Engine/Plugins/Developer/Concert/ConcertSync/ConcertSyncClient/Source/ConcertSyncClient/Private/ConcertClientPackageManager.h
+++ b/Engine/Plugins/Developer/Concert/ConcertSync/ConcertSyncClient/Source/ConcertSyncClient/Private/ConcertClientPackageManager.h
@@ -90,6 +90,11 @@ public:
 	 */
 	bool HasSessionChanges() const;
 
+	/** 
+	 * Returns the full path to the package if it is valid.
+	 */
+	TOptional<FString> GetValidPackageSessionPath(FName PackageName) const;
+
 	/**
 	 * Persist the session changes from the package name list and prepare it for source control submission.
 	 */
@@ -100,7 +105,6 @@ public:
 	 */
 	FOnConcertClientPackageTooLargeError& OnConcertClientPackageTooLargeError() { return OnPackageTooLargeErrorDelegate; }
 
-
 	/**
 	 * Returns true if the named package is participating in a package reload.
 	 */
@@ -108,6 +112,11 @@ public:
 
 private:
 
+	/**
+	 * Returns a full path for given package name if it was deleted and exists on the non-sandbox area.
+	 */
+	TOptional<FString> GetDeletedPackagePath(FName PackageName) const;
+
 	/**
 	 * Apply the package filters on the package info
 	 * @return true if the package info passes the filter.
diff --git a/Engine/Plugins/Developer/Concert/ConcertSync/ConcertSyncClient/Source/ConcertSyncClient/Private/ConcertClientWorkspace.cpp b/Engine/Plugins/Developer/Concert/ConcertSync/ConcertSyncClient/Source/ConcertSyncClient/Private/ConcertClientWorkspace.cpp
--- a/Engine/Plugins/Developer/Concert/ConcertSync/ConcertSyncClient/Source/ConcertSyncClient/Private/ConcertClientWorkspace.cpp
+++ b/Engine/Plugins/Developer/Concert/ConcertSync/ConcertSyncClient/Source/ConcertSyncClient/Private/ConcertClientWorkspace.cpp
@@ -276,6 +276,17 @@ TArray<FName> FConcertClientWorkspace::GatherSessionChanges(bool IgnorePersisted
 	return SessionChangedPackageNames.Array();
 }
 
+TOptional<FString> FConcertClientWorkspace::GetValidPackageSessionPath(FName PackageName) const
+{
+#if WITH_EDITOR
+	if (PackageManager)
+	{
+		return PackageManager->GetValidPackageSessionPath(MoveTemp(PackageName));
+	}
+#endif
+	return {};
+}
+
 bool FConcertClientWorkspace::PersistSessionChanges(TArrayView<const FName> InPackagesToPersist, ISourceControlProvider* SourceControlProvider, TArray<FText>* OutFailureReasons)
 {
 	bool bSuccess = false;
diff --git a/Engine/Plugins/Developer/Concert/ConcertSync/ConcertSyncClient/Source/ConcertSyncClient/Private/ConcertClientWorkspace.h b/Engine/Plugins/Developer/Concert/ConcertSync/ConcertSyncClient/Source/ConcertSyncClient/Private/ConcertClientWorkspace.h
--- a/Engine/Plugins/Developer/Concert/ConcertSync/ConcertSyncClient/Source/ConcertSyncClient/Private/ConcertClientWorkspace.h
+++ b/Engine/Plugins/Developer/Concert/ConcertSync/ConcertSyncClient/Source/ConcertSyncClient/Private/ConcertClientWorkspace.h
@@ -39,6 +39,7 @@ public:
 	virtual TFuture<FConcertResourceLockResponse> UnlockResources(TArray<FName> InResourceNames) override;
 	virtual bool HasSessionChanges() const override;
 	virtual TArray<FName> GatherSessionChanges(bool IgnorePersisted = true) override;
+	virtual TOptional<FString> GetValidPackageSessionPath(FName PackageName) const override;
 	virtual bool PersistSessionChanges(TArrayView<const FName> InPackagesToPersist, ISourceControlProvider* SourceControlProvider, TArray<FText>* OutFailureReasons = nullptr) override;
 	virtual bool HasLiveTransactionSupport(UPackage* InPackage) const override;
 	virtual bool ShouldIgnorePackageDirtyEvent(class UPackage* InPackage) const override;
diff --git a/Engine/Plugins/Developer/Concert/ConcertSync/ConcertSyncClient/Source/ConcertSyncClient/Private/ConcertSandboxPlatformFile.cpp b/Engine/Plugins/Developer/Concert/ConcertSync/ConcertSyncClient/Source/ConcertSyncClient/Private/ConcertSandboxPlatformFile.cpp
--- a/Engine/Plugins/Developer/Concert/ConcertSync/ConcertSyncClient/Source/ConcertSyncClient/Private/ConcertSandboxPlatformFile.cpp
+++ b/Engine/Plugins/Developer/Concert/ConcertSync/ConcertSyncClient/Source/ConcertSyncClient/Private/ConcertSandboxPlatformFile.cpp
@@ -187,6 +187,19 @@ const TCHAR* FConcertSandboxPlatformFile::GetName() const
 	return GetTypeName();
 }
 
+bool FConcertSandboxPlatformFile::DeletedPackageExistsInNonSandbox(FString InFilename) const
+{
+	const FConcertSandboxPlatformFilePath ResolvedPath = ToSandboxPath(MoveTemp(InFilename));
+	if (ResolvedPath.HasSandboxPath())
+	{
+		if (IsPathDeleted(ResolvedPath))
+		{
+			return LowerLevel->FileExists(*ResolvedPath.GetNonSandboxPath());
+		}
+	}
+	return false;
+}
+
 bool FConcertSandboxPlatformFile::FileExists(const TCHAR* Filename)
 {
 	const FConcertSandboxPlatformFilePath ResolvedPath = ToSandboxPath(Filename);
@@ -1131,7 +1144,7 @@ void FConcertSandboxPlatformFile::UnregisterContentMountPath(const FString& InCo
 
 FConcertSandboxPlatformFilePath FConcertSandboxPlatformFile::ToSandboxPath(FString InFilename, const bool bEvenIfDisabled) const
 {
-	return ToSandboxPath_Absolute(FPaths::ConvertRelativePathToFull(InFilename), bEvenIfDisabled);
+	return ToSandboxPath_Absolute(FPaths::ConvertRelativePathToFull(MoveTemp(InFilename)), bEvenIfDisabled);
 }
 
 FConcertSandboxPlatformFilePath FConcertSandboxPlatformFile::ToSandboxPath_Absolute(FString InFilename, const bool bEvenIfDisabled) const
@@ -1157,7 +1170,7 @@ FConcertSandboxPlatformFilePath FConcertSandboxPlatformFile::ToSandboxPath_Absol
 
 FConcertSandboxPlatformFilePath FConcertSandboxPlatformFile::FromSandboxPath(FString InFilename) const
 {
-	return FromSandboxPath_Absolute(FPaths::ConvertRelativePathToFull(InFilename));
+	return FromSandboxPath_Absolute(FPaths::ConvertRelativePathToFull(MoveTemp(InFilename)));
 }
 
 FConcertSandboxPlatformFilePath FConcertSandboxPlatformFile::FromSandboxPath_Absolute(FString InFilename) const
diff --git a/Engine/Plugins/Developer/Concert/ConcertSync/ConcertSyncClient/Source/ConcertSyncClient/Private/ConcertSandboxPlatformFile.h b/Engine/Plugins/Developer/Concert/ConcertSync/ConcertSyncClient/Source/ConcertSyncClient/Private/ConcertSandboxPlatformFile.h
--- a/Engine/Plugins/Developer/Concert/ConcertSync/ConcertSyncClient/Source/ConcertSyncClient/Private/ConcertSandboxPlatformFile.h
+++ b/Engine/Plugins/Developer/Concert/ConcertSync/ConcertSyncClient/Source/ConcertSyncClient/Private/ConcertSandboxPlatformFile.h
@@ -200,6 +200,11 @@ public:
 	 */
 	TArray<FString> GatherSandboxChangedFilenames() const;
 
+	/**
+	 * Returns true if the given package file exists on non sandbox path.
+	 */
+	bool DeletedPackageExistsInNonSandbox(FString InFilename) const;
+
 private:
 	struct FDirectoryItem
 	{
diff --git a/Engine/Plugins/Developer/Concert/ConcertSync/ConcertSyncClient/Source/ConcertSyncClient/Public/IConcertClientWorkspace.h b/Engine/Plugins/Developer/Concert/ConcertSync/ConcertSyncClient/Source/ConcertSyncClient/Public/IConcertClientWorkspace.h
--- a/Engine/Plugins/Developer/Concert/ConcertSync/ConcertSyncClient/Source/ConcertSyncClient/Public/IConcertClientWorkspace.h
+++ b/Engine/Plugins/Developer/Concert/ConcertSync/ConcertSyncClient/Source/ConcertSyncClient/Public/IConcertClientWorkspace.h
@@ -107,6 +107,13 @@ public:
 	 */
 	virtual TArray<FName> GatherSessionChanges(bool IgnorePersisted = true) = 0;
 
+	/**
+	 * Returns the full path to the package if it is a valid package session change.
+	 * @param PackageName the package to check.
+	 * @return the full path to the package if it is valid.
+	 */
+	virtual TOptional<FString> GetValidPackageSessionPath(FName PackageName) const = 0;
+
 	/** Persist the session changes from the package list and prepare it for source control submission */
 	virtual bool PersistSessionChanges(TArrayView<const FName> InPackageToPersist, ISourceControlProvider* SourceControlProvider, TArray<FText>* OutFailureReasonMap = nullptr) = 0;
 
diff --git a/Engine/Plugins/Developer/Concert/MultiUserClient/Source/MultiUserClient/Private/ConcertWorkspaceUI.cpp b/Engine/Plugins/Developer/Concert/MultiUserClient/Source/MultiUserClient/Private/ConcertWorkspaceUI.cpp
--- a/Engine/Plugins/Developer/Concert/MultiUserClient/Source/MultiUserClient/Private/ConcertWorkspaceUI.cpp
+++ b/Engine/Plugins/Developer/Concert/MultiUserClient/Source/MultiUserClient/Private/ConcertWorkspaceUI.cpp
@@ -606,23 +606,19 @@ bool FConcertWorkspaceUI::PromptPersistSessionChanges()
 	if (ClientWorkspacePin.IsValid())
 	{
 		// Get source control status of packages 
-		TArray<FName> PackageNames;
 		TArray<FString> PackageFilenames;
+		TArray<FName> PackageNames;
 		TArray<FSourceControlStateRef> States;
-		PackageNames = ClientWorkspacePin->GatherSessionChanges();
-		FString Filename;
-		for (auto It = PackageNames.CreateIterator(); It; ++It)
+		TArray<FName> CandidatePackages = ClientWorkspacePin->GatherSessionChanges();
+		for (FName PackageName : CandidatePackages)
 		{
-			if (FPackageName::DoesPackageExist(It->ToString(), &Filename))
-			{
-				PackageFilenames.Add(FPaths::ConvertRelativePathToFull(MoveTemp(Filename)));
-			}
-			// if the package file does not exist locally, remove it from the persist list, the db contains transaction data on file not propagated through Multi-User.
-			else
+			if (TOptional<FString> PackagePath = ClientWorkspacePin->GetValidPackageSessionPath(PackageName))
 			{
-				It.RemoveCurrent();
+				PackageFilenames.Add(PackagePath.GetValue());
+				PackageNames.Add(PackageName);
 			}
 		}
+
 		ECommandResult::Type Result = ISourceControlModule::Get().GetProvider().GetState(PackageFilenames, States, EStateCacheUsage::ForceUpdate);
 		// The dummy Multi-User source control provider always succeed and always return proxy states.
 		ensure(Result == ECommandResult::Succeeded);
@@ -633,7 +629,7 @@ bool FConcertWorkspaceUI::PromptPersistSessionChanges()
 			PersistItems.Add(MakeShared<FConcertPersistItem>(PackageNames[Index], States[Index]));
 		}
 	}
-	
+
 	TSharedRef<SWindow> NewWindow = SNew(SWindow)
 		.Title(LOCTEXT("PersistSubmitWindowTitle", "Persist & Submit Files"))
 		.SizingRule(ESizingRule::UserSized)
