Add make the default value of reflect visibility in game to be false and use SB to set the value on launch.

diff --git a/Engine/Plugins/Developer/Concert/ConcertMain/Source/Concert/Public/ConcertSettings.h b/Engine/Plugins/Developer/Concert/ConcertMain/Source/Concert/Public/ConcertSettings.h
--- a/Engine/Plugins/Developer/Concert/ConcertMain/Source/Concert/Public/ConcertSettings.h
+++ b/Engine/Plugins/Developer/Concert/ConcertMain/Source/Concert/Public/ConcertSettings.h
@@ -243,9 +243,12 @@ struct FConcertClientSettings
 	UPROPERTY(config, EditAnywhere, DisplayName="Latency Compensation", AdvancedDisplay, Category="Client Settings", meta=(ForceUnits=ms))
 	float LatencyCompensationMs;
 
-	/** When level editor changes are made, reflect those changes to the game equivalent property */
+	/**
+	 * When level editor changes are made, reflect those changes to the game equivalent property.
+	 * This settings can be specified on the editor cmd with `-CONCERTREFLECTVISIBILITY=`.
+	 */
 	UPROPERTY(config, EditAnywhere, DisplayName="Reflect Level Visibility to Game", AdvancedDisplay, Category="Client Settings")
-	bool bReflectLevelEditorInGame = true;
+	bool bReflectLevelEditorInGame = false;
 
 	/** Array of tags that can be used for grouping and categorizing. */
 	UPROPERTY(config, EditAnywhere, AdvancedDisplay, Category = "Client Settings")
diff --git a/Engine/Plugins/Developer/Concert/ConcertSync/ConcertSyncClient/Source/ConcertSyncClient/Private/ConcertClientWorkspace.cpp b/Engine/Plugins/Developer/Concert/ConcertSync/ConcertSyncClient/Source/ConcertSyncClient/Private/ConcertClientWorkspace.cpp
--- a/Engine/Plugins/Developer/Concert/ConcertSync/ConcertSyncClient/Source/ConcertSyncClient/Private/ConcertClientWorkspace.cpp
+++ b/Engine/Plugins/Developer/Concert/ConcertSync/ConcertSyncClient/Source/ConcertSyncClient/Private/ConcertClientWorkspace.cpp
@@ -125,7 +125,10 @@ void SetReflectEditorLevelVisibilityWithGame(bool InValue)
 	// Detail mode was modified, so store in the CVar
 	static IConsoleVariable* CVar = IConsoleManager::Get().FindConsoleVariable(TEXT("Editor.ReflectEditorLevelVisibilityWithGame"));
 	int32 ValAsInt = !!InValue;
-	CVar->Set(ValAsInt);
+	if (CVar->GetInt() != ValAsInt)
+	{
+		CVar->Set(ValAsInt);
+	}
 #endif
 }
 
@@ -193,11 +196,17 @@ struct FConcertWorkspaceConsoleCommands
 	FAutoConsoleCommand DisableRemoteVerboseLogging;
 };
 
-FConcertClientWorkspace::FConcertClientWorkspace(TSharedRef<FConcertSyncClientLiveSession> InLiveSession, IConcertClientPackageBridge* InPackageBridge, IConcertClientTransactionBridge* InTransactionBridge, TSharedPtr<IConcertFileSharingService> InFileSharingService)
-	: FileSharingService(MoveTemp(InFileSharingService))
+FConcertClientWorkspace::FConcertClientWorkspace(TSharedRef<FConcertSyncClientLiveSession> InLiveSession,
+												 IConcertClientPackageBridge* InPackageBridge,
+												 IConcertClientTransactionBridge* InTransactionBridge,
+												 TSharedPtr<IConcertFileSharingService> InFileSharingService,
+												 IConcertSyncClient* InOwnerSyncClient)
+	: OwnerSyncClient(InOwnerSyncClient),
+	  FileSharingService(MoveTemp(InFileSharingService))
 {
 	static FConcertWorkspaceConsoleCommands ConsoleCommands;
 
+	check(OwnerSyncClient);
 	BindSession(InLiveSession, InPackageBridge, InTransactionBridge);
 }
 
@@ -845,8 +854,8 @@ void FConcertClientWorkspace::OnEndFrame()
 		}
 	}
 	LiveSession->GetSessionDatabase().UpdateAsynchronousTasks();
-	const UConcertClientConfig *Config = GetDefault<UConcertClientConfig>();
-	SetReflectEditorLevelVisibilityWithGame(Config->ClientSettings.bReflectLevelEditorInGame);
+	IConcertClientRef ConcertClient = OwnerSyncClient->GetConcertClient();
+	SetReflectEditorLevelVisibilityWithGame(ConcertClient->GetConfiguration()->ClientSettings.bReflectLevelEditorInGame);
 }
 
 void FConcertClientWorkspace::HandleWorkspaceSyncEndpointEvent(const FConcertSessionContext& Context, const FConcertWorkspaceSyncEndpointEvent& Event)
diff --git a/Engine/Plugins/Developer/Concert/ConcertSync/ConcertSyncClient/Source/ConcertSyncClient/Private/ConcertClientWorkspace.h b/Engine/Plugins/Developer/Concert/ConcertSync/ConcertSyncClient/Source/ConcertSyncClient/Private/ConcertClientWorkspace.h
--- a/Engine/Plugins/Developer/Concert/ConcertSync/ConcertSyncClient/Source/ConcertSyncClient/Private/ConcertClientWorkspace.h
+++ b/Engine/Plugins/Developer/Concert/ConcertSync/ConcertSyncClient/Source/ConcertSyncClient/Private/ConcertClientWorkspace.h
@@ -3,6 +3,7 @@
 #pragma once
 
 #include "CoreMinimal.h"
+#include "IConcertSyncClient.h"
 #include "Misc/ITransaction.h"
 #include "IConcertClientWorkspace.h"
 #include "IConcertSessionHandler.h"
@@ -27,7 +28,7 @@ struct FScopedSlowTask;
 class FConcertClientWorkspace : public IConcertClientWorkspace
 {
 public:
-	FConcertClientWorkspace(TSharedRef<FConcertSyncClientLiveSession> InLiveSession, IConcertClientPackageBridge* InPackageBridge, IConcertClientTransactionBridge* InTransactionBridge, TSharedPtr<IConcertFileSharingService> InFileSharingService);
+	FConcertClientWorkspace(TSharedRef<FConcertSyncClientLiveSession> InLiveSession, IConcertClientPackageBridge* InPackageBridge, IConcertClientTransactionBridge* InTransactionBridge, TSharedPtr<IConcertFileSharingService> InFileSharingService, IConcertSyncClient* InOwnerSyncClient);
 	virtual ~FConcertClientWorkspace();
 
 	// IConcertClientWorkspace interface
@@ -208,6 +209,9 @@ private:
 	/** */
 	IConcertClientPackageBridge* PackageBridge;
 
+	/** The sync client that owns the workspace. */
+	IConcertSyncClient *OwnerSyncClient;
+
 	/** */
 	TSharedPtr<FConcertSyncClientLiveSession> LiveSession;
 
diff --git a/Engine/Plugins/Developer/Concert/ConcertSync/ConcertSyncClient/Source/ConcertSyncClient/Private/ConcertSyncClient.cpp b/Engine/Plugins/Developer/Concert/ConcertSync/ConcertSyncClient/Source/ConcertSyncClient/Private/ConcertSyncClient.cpp
--- a/Engine/Plugins/Developer/Concert/ConcertSync/ConcertSyncClient/Source/ConcertSyncClient/Private/ConcertSyncClient.cpp
+++ b/Engine/Plugins/Developer/Concert/ConcertSync/ConcertSyncClient/Source/ConcertSyncClient/Private/ConcertSyncClient.cpp
@@ -134,7 +134,7 @@ void FConcertSyncClient::SetFileSharingService(TSharedPtr<IConcertFileSharingSer
 void FConcertSyncClient::CreateWorkspace(const TSharedRef<FConcertSyncClientLiveSession>& InLiveSession)
 {
 	DestroyWorkspace();
-	Workspace = MakeShared<FConcertClientWorkspace>(InLiveSession, PackageBridge, TransactionBridge, FileSharingService);
+	Workspace = MakeShared<FConcertClientWorkspace>(InLiveSession, PackageBridge, TransactionBridge, FileSharingService, this);
 	OnWorkspaceStartupDelegate.Broadcast(Workspace);
 #if WITH_EDITOR
 	if (GIsEditor && EnumHasAllFlags(SessionFlags, EConcertSyncSessionFlags::EnablePackages | EConcertSyncSessionFlags::ShouldUsePackageSandbox))
diff --git a/Engine/Plugins/Developer/Concert/ConcertSync/ConcertSyncClient/Source/ConcertSyncClient/Private/ConcertSyncClientModule.cpp b/Engine/Plugins/Developer/Concert/ConcertSync/ConcertSyncClient/Source/ConcertSyncClient/Private/ConcertSyncClientModule.cpp
--- a/Engine/Plugins/Developer/Concert/ConcertSync/ConcertSyncClient/Source/ConcertSyncClient/Private/ConcertSyncClientModule.cpp
+++ b/Engine/Plugins/Developer/Concert/ConcertSync/ConcertSyncClient/Source/ConcertSyncClient/Private/ConcertSyncClientModule.cpp
@@ -69,6 +69,9 @@ public:
 			ClientConfig->bRetryAutoConnectOnError |= FParse::Param(CommandLine, TEXT("CONCERTRETRYAUTOCONNECTONERROR"));
 			FParse::Bool(CommandLine, TEXT("-CONCERTRETRYAUTOCONNECTONERROR="), ClientConfig->bRetryAutoConnectOnError);
 
+			ClientConfig->ClientSettings.bReflectLevelEditorInGame |= FParse::Param(CommandLine, TEXT("CONCERTREFLECTVISIBILITY"));
+			FParse::Bool(CommandLine, TEXT("-CONCERTREFLECTVISIBILITY="), ClientConfig->ClientSettings.bReflectLevelEditorInGame);
+
 			ClientConfig->bIsHeadless |= FParse::Param(CommandLine, TEXT("CONCERTISHEADLESS"));
 			FParse::Bool(CommandLine, TEXT("-CONCERTISHEADLESS="), ClientConfig->bIsHeadless);
 			ClientConfig->bIsHeadless |= !FSlateNotificationManager::Get().AreNotificationsAllowed();
diff --git a/Engine/Plugins/VirtualProduction/Switchboard/Source/Switchboard/switchboard/devices/unreal/plugin_unreal.py b/Engine/Plugins/VirtualProduction/Switchboard/Source/Switchboard/switchboard/devices/unreal/plugin_unreal.py
--- a/Engine/Plugins/VirtualProduction/Switchboard/Source/Switchboard/switchboard/devices/unreal/plugin_unreal.py
+++ b/Engine/Plugins/VirtualProduction/Switchboard/Source/Switchboard/switchboard/devices/unreal/plugin_unreal.py
@@ -380,6 +380,16 @@ class DeviceUnreal(Device):
                 'Directory in which to store logs transferred from devices. '
                 'If unset, defaults to $(ProjectDir)/Saved/Logs/Switchboard/'),
         ),
+        'reflect_visibility_to_game': BoolSetting(
+            attr_name='reflect_visibility_to_game',
+            nice_name='Reflect Editor Visibility to Game',
+            value=True,
+            tool_tip=(
+                'Sets the value for `Reflect Level Visibilty to Game` for Multi-user editing. \n'
+                'Editor visibilty state will be applied to the game equivalent visibility properties. \n'
+                'This is useful for ICVFX workflows where the editor visibilty state directly \n'
+                'corresponds to the state on the render nodes.')
+        ),
         'rsync_port': IntSetting(
             attr_name='rsync_port',
             nice_name='Rsync Server Port',
@@ -1208,6 +1218,10 @@ class DeviceUnreal(Device):
 
         return vproles, missing_roles
 
+    @property
+    def reflect_editor_vis_in_game(self) -> str:
+        return DeviceUnreal.csettings['reflect_visibility_to_game'].get_value()
+
     @property
     def udpmessaging_multicast_endpoint(self) -> str:
         device_multicast = DeviceUnreal.csettings[
@@ -1319,6 +1333,9 @@ class DeviceUnreal(Device):
                 self.name):
             command_line_args += ' -AutoDeclinePackageRecovery'
 
+        concert_vis_reflect = 1 if self.reflect_editor_vis_in_game else 0
+        command_line_args += f' -ConcertReflectVisibility={concert_vis_reflect}'
+
         # UdpMessaging endpoints
         if self.udpmessaging_multicast_endpoint:
             server_mc = CONFIG.MUSERVER_MULTICAST_ENDPOINT.get_value().strip()
