Add a time frame adjustment event to capture changes to the underlying MovieScene between asset saves.

diff --git a/Engine/Plugins/Developer/Concert/ConcertSync/ConcertSyncClient/Source/ConcertSyncClient/Private/ConcertClientSequencerManager.cpp b/Engine/Plugins/Developer/Concert/ConcertSync/ConcertSyncClient/Source/ConcertSyncClient/Private/ConcertClientSequencerManager.cpp
--- a/Engine/Plugins/Developer/Concert/ConcertSync/ConcertSyncClient/Source/ConcertSyncClient/Private/ConcertClientSequencerManager.cpp
+++ b/Engine/Plugins/Developer/Concert/ConcertSync/ConcertSyncClient/Source/ConcertSyncClient/Private/ConcertClientSequencerManager.cpp
@@ -1,6 +1,7 @@
 // Copyright Epic Games, Inc. All Rights Reserved.
 
 #include "ConcertClientSequencerManager.h"
+#include "ConcertSequencerMessages.h"
 #include "Delegates/IDelegateInstance.h"
 #include "HAL/IConsoleManager.h"
 #include "Misc/QualifiedFrameTime.h"
@@ -133,7 +134,8 @@ void FConcertClientSequencerManager::OnSequencerCreated(TSharedRef<ISequencer> I
 	}
 }
 
-TArray<FConcertClientSequencerManager::FOpenSequencerData*, TInlineAllocator<1>> FConcertClientSequencerManager::GatherRootSequencersByState(const FConcertSequencerState& InSequenceState)
+TArray<FConcertClientSequencerManager::FOpenSequencerData*, TInlineAllocator<1>>
+FConcertClientSequencerManager::GatherRootSequencersByState(const FString& InSequenceObjectPath)
 {
 	TArray<FOpenSequencerData*, TInlineAllocator<1>> OutSequencers;
 	for (FOpenSequencerData& Entry : OpenSequencers)
@@ -141,7 +143,7 @@ TArray<FConcertClientSequencerManager::FOpenSequencerData*, TInlineAllocator<1>>
 		TSharedPtr<ISequencer> Sequencer = Entry.WeakSequencer.Pin();
 		UMovieSceneSequence*   Sequence = Sequencer.IsValid() ? Sequencer->GetRootMovieSceneSequence() : nullptr;
 
-		if (Sequence && (Sequence->GetPathName() == InSequenceState.SequenceObjectPath || CVarEnableUnrelatedTimelineSync.GetValueOnAnyThread() > 0))
+		if (Sequence && (Sequence->GetPathName() == InSequenceObjectPath || CVarEnableUnrelatedTimelineSync.GetValueOnAnyThread() > 0))
 		{
 			OutSequencers.Add(&Entry);
 		}
@@ -172,6 +174,7 @@ void FConcertClientSequencerManager::Register(TSharedRef<IConcertClientSession>
 	InSession->RegisterCustomEventHandler<FConcertSequencerCloseEvent>(this, &FConcertClientSequencerManager::OnCloseEvent);
 	InSession->RegisterCustomEventHandler<FConcertSequencerOpenEvent>(this, &FConcertClientSequencerManager::OnOpenEvent);
 	InSession->RegisterCustomEventHandler<FConcertSequencerStateSyncEvent>(this, &FConcertClientSequencerManager::OnSyncEvent);
+	InSession->RegisterCustomEventHandler<FConcertSequencerTimeAdjustmentEvent>(this, &FConcertClientSequencerManager::OnTimeAdjustmentEvent);
 }
 
 void FConcertClientSequencerManager::Unregister(TSharedRef<IConcertClientSession> InSession)
@@ -184,6 +187,7 @@ void FConcertClientSequencerManager::Unregister(TSharedRef<IConcertClientSession
 		Session->UnregisterCustomEventHandler<FConcertSequencerCloseEvent>(this);
 		Session->UnregisterCustomEventHandler<FConcertSequencerOpenEvent>(this);
 		Session->UnregisterCustomEventHandler<FConcertSequencerStateSyncEvent>(this);
+		Session->UnregisterCustomEventHandler<FConcertSequencerTimeAdjustmentEvent>(this);
 	}
 	WeakSession.Reset();
 }
@@ -280,7 +284,7 @@ void FConcertClientSequencerManager::OnSyncEvent(const FConcertSessionContext& I
 	{
 		FConcertSequencerState& SequencerState = SequencerStates.FindOrAdd(*State.SequenceObjectPath);
 		SequencerState = State;
-		for (FOpenSequencerData* OpenSequencer : GatherRootSequencersByState(SequencerState))
+		for (FOpenSequencerData* OpenSequencer : GatherRootSequencersByState(SequencerState.SequenceObjectPath))
 		{
 			TSharedPtr<ISequencer> Sequencer = OpenSequencer->WeakSequencer.Pin();
 			if (Sequencer.IsValid() && IsSequencerPlaybackSyncEnabled())
@@ -395,7 +399,7 @@ void FConcertClientSequencerManager::ApplyTransportCloseEvent(const FConcertSequ
 		if (PendingClose.bControllerClose)
 		{
 			SequencerState->PlayerStatus = EConcertMovieScenePlayerStatus::Stopped;
-			for (FOpenSequencerData* OpenSequencer : GatherRootSequencersByState(*SequencerState))
+			for (FOpenSequencerData* OpenSequencer : GatherRootSequencersByState(SequencerState->SequenceObjectPath))
 			{
 				OpenSequencer->PlaybackMode = EPlaybackMode::Undefined;
 				TSharedPtr<ISequencer> Sequencer = OpenSequencer->WeakSequencer.Pin();
@@ -554,6 +558,101 @@ void FConcertClientSequencerManager::OnTransportEvent(const FConcertSessionConte
 	PendingSequencerEvents.Add(InEvent.State);
 }
 
+void FConcertClientSequencerManager::OnTimeAdjustmentEvent(const FConcertSessionContext&, const FConcertSequencerTimeAdjustmentEvent& TimeAdjustmentEvent)
+{
+	PendingTimeAdjustmentEvents.Add(TimeAdjustmentEvent);
+}
+
+void FConcertClientSequencerManager::ApplyTimeAdjustmentEvent(const FConcertSequencerTimeAdjustmentEvent& TimeEvent)
+{
+	if (bRespondingToTransportEvent)
+	{
+		return;
+	}
+
+	UE_LOG(LogConcertSequencerSync, Verbose, TEXT("Apply Time adjustment Event: %s"), *TimeEvent.SequenceObjectPath);
+	TGuardValue<bool> ReentrancyGuard(bRespondingToTransportEvent, true);
+
+	if (GIsEditor)
+	{
+		ApplyTimeAdjustmentSequencers(TimeEvent);
+	}
+	else if (CVarEnableSequencePlayer.GetValueOnAnyThread() > 0)
+	{
+		ApplyTimeAdjustmentPlayers(TimeEvent);
+	}
+}
+
+bool ApplyStartFrameToMovieScene(FFrameNumber StartFrame, UMovieScene* MovieScene)
+{
+	check(MovieScene);
+
+	FFrameNumber DeltaFrame = StartFrame - MovieScene->GetPlaybackRange().GetLowerBoundValue();
+	if (DeltaFrame == 0)
+	{
+		return false;
+	}
+
+	for (UMovieSceneSection* Section : MovieScene->GetAllSections())
+	{
+		Section->ComputeEffectiveRange();
+		Section->MoveSection(DeltaFrame);
+	}
+	MovieScene->SetPlaybackRange(TRange<FFrameNumber>(StartFrame, TNumericLimits<int32>::Max() - 1), false);
+
+	return true;
+}
+
+void FConcertClientSequencerManager::ApplyTimeAdjustmentSequencers(const FConcertSequencerTimeAdjustmentEvent& TimeEvent)
+{
+	for (FOpenSequencerData* OpenSequencer : GatherRootSequencersByState(TimeEvent.SequenceObjectPath))
+	{
+		ISequencer* Sequencer = OpenSequencer->WeakSequencer.Pin().Get();
+		// If the entry is driving playback (PlaybackMode == Controller) then we never respond time adjustment events.
+		if (!Sequencer || OpenSequencer->PlaybackMode == EPlaybackMode::Controller)
+		{
+			continue;
+		}
+
+		// Adjust the range of the sequencer based on the time provided.
+		//
+		UMovieScene* MovieScene = Sequencer->GetRootMovieSceneSequence()->GetMovieScene();
+		if (MovieScene)
+		{
+			if (ApplyStartFrameToMovieScene(TimeEvent.PlaybackStartFrame, MovieScene))
+			{
+				FFrameRate FrameRate = MovieScene->GetTickResolution();
+				double CurrentTimeSeconds = FrameRate.AsSeconds(FFrameTime(TimeEvent.PlaybackStartFrame));
+				TRange<double> NewRange(CurrentTimeSeconds, CurrentTimeSeconds+10.0);
+				Sequencer->SetViewRange(NewRange, EViewRangeInterpolation::Immediate);
+				Sequencer->SetClampRange(Sequencer->GetViewRange());
+			}
+		}
+
+	}
+}
+
+void FConcertClientSequencerManager::ApplyTimeAdjustmentPlayers(const FConcertSequencerTimeAdjustmentEvent& TimeEvent)
+{
+	FSequencePlayer* SeqPlayer = SequencePlayers.Find(*TimeEvent.SequenceObjectPath);
+
+	if (!SeqPlayer)
+	{
+		return;
+	}
+
+	ALevelSequenceActor* LevelSequenceActor = SeqPlayer->Actor.Get();
+	if (LevelSequenceActor && LevelSequenceActor->SequencePlayer)
+	{
+		UMovieScene* MovieScene = LevelSequenceActor->LevelSequenceAsset->GetMovieScene();
+		if (MovieScene)
+		{
+			ApplyStartFrameToMovieScene(TimeEvent.PlaybackStartFrame, MovieScene);
+		}
+	}
+}
+
+
 void FConcertClientSequencerManager::ApplyTransportEvent(const FConcertSequencerState& EventState)
 {
 	if (bRespondingToTransportEvent)
@@ -591,7 +690,7 @@ void FConcertClientSequencerManager::ApplyEventToSequencers(const FConcertSequen
 	float LatencyCompensationMs = GetLatencyCompensationMs();
 
 	// Update all opened sequencer with this root sequence
-	for (FOpenSequencerData* OpenSequencer : GatherRootSequencersByState(EventState))
+	for (FOpenSequencerData* OpenSequencer : GatherRootSequencersByState(EventState.SequenceObjectPath))
 	{
 		ISequencer* Sequencer = OpenSequencer->WeakSequencer.Pin().Get();
 		// If the entry is driving playback (PlaybackMode == Controller) then we never respond to external transport events
@@ -834,6 +933,12 @@ void FConcertClientSequencerManager::OnEndFrame()
 		ApplyTransportEvent(State);
 	}
 	PendingSequencerEvents.Reset();
+
+	for (const FConcertSequencerTimeAdjustmentEvent& Event : PendingTimeAdjustmentEvents)
+	{
+		ApplyTimeAdjustmentEvent(Event);
+	}
+	PendingTimeAdjustmentEvents.Reset();
 }
 
 void FConcertClientSequencerManager::AddReferencedObjects(FReferenceCollector& Collector)
diff --git a/Engine/Plugins/Developer/Concert/ConcertSync/ConcertSyncClient/Source/ConcertSyncClient/Private/ConcertClientSequencerManager.h b/Engine/Plugins/Developer/Concert/ConcertSync/ConcertSyncClient/Source/ConcertSyncClient/Private/ConcertClientSequencerManager.h
--- a/Engine/Plugins/Developer/Concert/ConcertSync/ConcertSyncClient/Source/ConcertSyncClient/Private/ConcertClientSequencerManager.h
+++ b/Engine/Plugins/Developer/Concert/ConcertSync/ConcertSyncClient/Source/ConcertSyncClient/Private/ConcertClientSequencerManager.h
@@ -184,6 +184,13 @@ private:
 	 */
 	void OnSyncEvent(const FConcertSessionContext& InEventContext, const FConcertSequencerStateSyncEvent& InEvent);
 
+	/**
+	 * Called on receipt of time adjustment event.
+	 * @param InEventContext             The context for the current session.
+	 * @param InEvent                    The sequencer time adjustment event.
+	 */
+	void OnTimeAdjustmentEvent(const FConcertSessionContext& InEventContext, const FConcertSequencerTimeAdjustmentEvent& TimeAdjustmentEvent);
+
 	/**
 	 * Called when the global time has been changed for the specified Sequencer
 	 *
@@ -252,13 +259,28 @@ private:
 	 */
 	void ApplyTransportCloseEvent(const FConcertSequencerCloseEvent& PendingClose);
 
+	/**
+	 * Apply a sequencer time adjustment event.
+	 */
+	void ApplyTimeAdjustmentEvent(const FConcertSequencerTimeAdjustmentEvent& TimeAdjustmentEvent);
+
+	/**
+	 * Apply a sequencer time adjustment to players.
+	 */
+	void ApplyTimeAdjustmentPlayers(const FConcertSequencerTimeAdjustmentEvent& TimeAdjustmentEvent);
+
+	/**
+	 * Apply a sequencer time adjustment to sequencers.
+	 */
+	void ApplyTimeAdjustmentSequencers(const FConcertSequencerTimeAdjustmentEvent& TimeAdjustmentEvent);
+
 	/**
 	 * Gather all the currently open sequencer UIs that have the specified path as their root sequence
 	 *
 	 * @param InSequenceObjectPath        The full path to the root asset to gather sequences for
 	 * @return An array containing all the entries that apply to the supplied sequence path
 	 */
-	TArray<FOpenSequencerData*, TInlineAllocator<1>> GatherRootSequencersByState(const FConcertSequencerState& InSequenceObjectPath);
+	TArray<FOpenSequencerData*, TInlineAllocator<1>> GatherRootSequencersByState(const FString& InSequenceObjectPath);
 
 	/**
 	 * Get the amount of latency compensation to apply to time-synchronization sensitive interactions
@@ -294,6 +316,9 @@ private:
 	/** List of pending sequencer open events to apply at end of frame. */
 	TArray<FConcertSequencerCloseEvent> PendingSequenceCloseEvents;
 
+	/** List of pending time adjustment events. */
+	TArray<FConcertSequencerTimeAdjustmentEvent> PendingTimeAdjustmentEvents;
+
 	/**
 	 * List of pending sequence players that are scheduled for destruction. The first member of the tuple is the key /
 	 * level sequence player we a removing.  The second member tuple is the soft object path name of the replacement
diff --git a/Engine/Plugins/Developer/Concert/ConcertSync/ConcertSyncCore/Source/ConcertSyncCore/Public/ConcertSequencerMessages.h b/Engine/Plugins/Developer/Concert/ConcertSync/ConcertSyncCore/Source/ConcertSyncCore/Public/ConcertSequencerMessages.h
--- a/Engine/Plugins/Developer/Concert/ConcertSync/ConcertSyncCore/Source/ConcertSyncCore/Public/ConcertSequencerMessages.h
+++ b/Engine/Plugins/Developer/Concert/ConcertSync/ConcertSyncCore/Source/ConcertSyncCore/Public/ConcertSequencerMessages.h
@@ -108,3 +108,19 @@ struct FConcertSequencerStateSyncEvent
 	UPROPERTY()
 	TArray<FConcertSequencerState> SequencerStates;
 };
+
+/**
+ * An event that represents a time changes on the sequencer. This can happen via take recorder
+ * which will shift the active take _if_ Start At Timecode is enabled.
+ */
+USTRUCT()
+struct FConcertSequencerTimeAdjustmentEvent
+{
+	GENERATED_BODY()
+
+	UPROPERTY()
+	FFrameNumber PlaybackStartFrame;
+
+	UPROPERTY()
+	FString SequenceObjectPath;
+};
diff --git a/Engine/Plugins/Developer/Concert/ConcertSync/ConcertSyncServer/Source/ConcertSyncServer/Private/ConcertServerSequencerManager.cpp b/Engine/Plugins/Developer/Concert/ConcertSync/ConcertSyncServer/Source/ConcertSyncServer/Private/ConcertServerSequencerManager.cpp
--- a/Engine/Plugins/Developer/Concert/ConcertSync/ConcertSyncServer/Source/ConcertSyncServer/Private/ConcertServerSequencerManager.cpp
+++ b/Engine/Plugins/Developer/Concert/ConcertSync/ConcertSyncServer/Source/ConcertSyncServer/Private/ConcertServerSequencerManager.cpp
@@ -25,6 +25,7 @@ void FConcertServerSequencerManager::BindSession(const TSharedRef<FConcertSyncSe
 	LiveSession->GetSession().RegisterCustomEventHandler<FConcertSequencerCloseEvent>(this, &FConcertServerSequencerManager::HandleSequencerCloseEvent);
 	LiveSession->GetSession().RegisterCustomEventHandler<FConcertSequencerStateEvent>(this, &FConcertServerSequencerManager::HandleSequencerStateEvent);
 	LiveSession->GetSession().RegisterCustomEventHandler<FConcertSequencerOpenEvent>(this, &FConcertServerSequencerManager::HandleSequencerOpenEvent);
+	LiveSession->GetSession().RegisterCustomEventHandler<FConcertSequencerTimeAdjustmentEvent>(this, &FConcertServerSequencerManager::HandleSequencerTimeAdjustmentEvent);
 }
 
 void FConcertServerSequencerManager::UnbindSession()
@@ -65,6 +66,20 @@ void FConcertServerSequencerManager::HandleSequencerOpenEvent(const FConcertSess
 	LiveSession->GetSession().SendCustomEvent(InEvent, ClientIds, EConcertMessageFlags::ReliableOrdered | EConcertMessageFlags::UniqueId);
 }
 
+void FConcertServerSequencerManager::HandleSequencerTimeAdjustmentEvent(const FConcertSessionContext& InEventContext, const FConcertSequencerTimeAdjustmentEvent& InEvent)
+{
+	// Verify that we have sequencers with the given SequenceObjectPath open.
+	//
+	FConcertOpenSequencerState* SequencerState = SequencerStates.Find(*InEvent.SequenceObjectPath);
+	if (SequencerState)
+	{
+		// Forward the message to the other clients
+		TArray<FGuid> ClientIds = LiveSession->GetSession().GetSessionClientEndpointIds();
+		ClientIds.Remove(InEventContext.SourceEndpointId);
+		LiveSession->GetSession().SendCustomEvent(InEvent, ClientIds, EConcertMessageFlags::ReliableOrdered | EConcertMessageFlags::UniqueId);
+	}
+}
+
 void FConcertServerSequencerManager::HandleSequencerCloseEvent(const FConcertSessionContext& InEventContext, const FConcertSequencerCloseEvent& InEvent)
 {
 	FConcertOpenSequencerState* SequencerState = SequencerStates.Find(*InEvent.SequenceObjectPath);
diff --git a/Engine/Plugins/Developer/Concert/ConcertSync/ConcertSyncServer/Source/ConcertSyncServer/Private/ConcertServerSequencerManager.h b/Engine/Plugins/Developer/Concert/ConcertSync/ConcertSyncServer/Source/ConcertSyncServer/Private/ConcertServerSequencerManager.h
--- a/Engine/Plugins/Developer/Concert/ConcertSync/ConcertSyncServer/Source/ConcertSyncServer/Private/ConcertServerSequencerManager.h
+++ b/Engine/Plugins/Developer/Concert/ConcertSync/ConcertSyncServer/Source/ConcertSyncServer/Private/ConcertServerSequencerManager.h
@@ -32,6 +32,9 @@ private:
 	/** Handler for the sequencer close event. */
 	void HandleSequencerCloseEvent(const FConcertSessionContext& InEventContext, const FConcertSequencerCloseEvent& InEvent);
 
+	/** Handler for the sequencer close event. */
+	void HandleSequencerTimeAdjustmentEvent(const FConcertSessionContext& InEventContext, const FConcertSequencerTimeAdjustmentEvent& InEvent);
+
 	/** Handler for the session clients changed event. */
 	void HandleSessionClientChanged(IConcertServerSession& InSession, EConcertClientStatus InClientStatus, const FConcertSessionClientInfo& InClientInfo);
 
diff --git a/Engine/Plugins/VirtualProduction/MultiUserTakes/Source/ConcertTakeRecorder/Private/ConcertTakeRecorderManager.cpp b/Engine/Plugins/VirtualProduction/MultiUserTakes/Source/ConcertTakeRecorder/Private/ConcertTakeRecorderManager.cpp
--- a/Engine/Plugins/VirtualProduction/MultiUserTakes/Source/ConcertTakeRecorder/Private/ConcertTakeRecorderManager.cpp
+++ b/Engine/Plugins/VirtualProduction/MultiUserTakes/Source/ConcertTakeRecorder/Private/ConcertTakeRecorderManager.cpp
@@ -5,6 +5,7 @@
 #include "Algo/Count.h"
 #include "Algo/Transform.h"
 #include "ConcertMessages.h"
+#include "ConcertSequencerMessages.h"
 #include "ConcertTransportMessages.h"
 #include "Delegates/IDelegateInstance.h"
 #include "HAL/IConsoleManager.h"
@@ -359,6 +360,7 @@ void FConcertTakeRecorderManager::OnTakeRecorderInitialized(UTakeRecorder* TakeR
 {
 	TakeRecorder->OnRecordingFinished().AddRaw(this, &FConcertTakeRecorderManager::OnRecordingFinished);
 	TakeRecorder->OnRecordingCancelled().AddRaw(this, &FConcertTakeRecorderManager::OnRecordingCancelled);
+	TakeRecorder->OnStartPlayFrameModified().AddRaw(this, &FConcertTakeRecorderManager::OnFrameAdjustment);
 
 	if (IsTakeSyncEnabled() && !bIsRecording && TakeRecorderState.LastStartedTake != TakeRecorder->GetName())
 	{
@@ -437,9 +439,24 @@ void FConcertTakeRecorderManager::OnRecordingCancelled(UTakeRecorder* TakeRecord
 
 	TakeRecorder->OnRecordingFinished().RemoveAll(this);
 	TakeRecorder->OnRecordingCancelled().RemoveAll(this);
+	TakeRecorder->OnStartPlayFrameModified().RemoveAll(this);
 	bIsRecording = false;
 }
 
+void FConcertTakeRecorderManager::OnFrameAdjustment(UTakeRecorder* TakeRecorder, const FFrameNumber& InPlaybackStartFrame)
+{
+	if (TSharedPtr<IConcertClientSession> Session = WeakSession.Pin())
+	{
+		ULevelSequence* LevelSequence = TakeRecorder->GetSequence();
+		check(LevelSequence);
+
+		FConcertSequencerTimeAdjustmentEvent TimeEvent;
+		TimeEvent.SequenceObjectPath = LevelSequence->GetPathName();
+		TimeEvent.PlaybackStartFrame = InPlaybackStartFrame;
+		Session->SendCustomEvent(TimeEvent, Session->GetSessionClientEndpointIds(), EConcertMessageFlags::ReliableOrdered);
+	}
+}
+
 void FConcertTakeRecorderManager::OnTakeInitializedEvent(const FConcertSessionContext&, const FConcertTakeInitializedEvent& InEvent)
 {
 	if (IsTakeSyncEnabled() && CanRecord())
diff --git a/Engine/Plugins/VirtualProduction/MultiUserTakes/Source/ConcertTakeRecorder/Private/ConcertTakeRecorderManager.h b/Engine/Plugins/VirtualProduction/MultiUserTakes/Source/ConcertTakeRecorder/Private/ConcertTakeRecorderManager.h
--- a/Engine/Plugins/VirtualProduction/MultiUserTakes/Source/ConcertTakeRecorder/Private/ConcertTakeRecorderManager.h
+++ b/Engine/Plugins/VirtualProduction/MultiUserTakes/Source/ConcertTakeRecorder/Private/ConcertTakeRecorderManager.h
@@ -6,6 +6,7 @@
 #include "ConcertMessages.h"
 #include "Delegates/IDelegateInstance.h"
 #include "HAL/Platform.h"
+#include "Misc/FrameNumber.h"
 #include "UObject/StrongObjectPtr.h"
 #include "TakePreset.h"
 #include "ConcertSyncClient/Public/IConcertClientTransactionBridge.h"
@@ -57,6 +58,7 @@ private:
 	void OnTakeRecorderInitialized(UTakeRecorder* TakeRecorder);
 	void OnRecordingFinished(UTakeRecorder* TakeRecorder);
 	void OnRecordingCancelled(UTakeRecorder* TakeRecorder);
+	void OnFrameAdjustment(UTakeRecorder* TakeRecorder, const FFrameNumber& InPlaybackStartFrame);
 
 	//~ Concert event handlers
 	void OnTakeInitializedEvent(const FConcertSessionContext&, const FConcertTakeInitializedEvent& InEvent);
diff --git a/Engine/Plugins/VirtualProduction/Takes/Source/TakeRecorder/Private/Recorder/TakeRecorder.cpp b/Engine/Plugins/VirtualProduction/Takes/Source/TakeRecorder/Private/Recorder/TakeRecorder.cpp
--- a/Engine/Plugins/VirtualProduction/Takes/Source/TakeRecorder/Private/Recorder/TakeRecorder.cpp
+++ b/Engine/Plugins/VirtualProduction/Takes/Source/TakeRecorder/Private/Recorder/TakeRecorder.cpp
@@ -976,6 +976,8 @@ void UTakeRecorder::Start()
 			}
 		}
 
+		OnFrameModifiedEvent.Broadcast(this, PlaybackStartFrame);
+
 		if (StopRecordingFrame.IsSet())
 		{
 			StopRecordingFrame = StopRecordingFrame.GetValue() + DeltaFrame;
@@ -1260,6 +1262,11 @@ FOnTakeRecordingCancelled& UTakeRecorder::OnRecordingCancelled()
 	return OnRecordingCancelledEvent;
 }
 
+FOnStartPlayFrameModified& UTakeRecorder::OnStartPlayFrameModified()
+{
+	return OnFrameModifiedEvent;
+}
+
 void UTakeRecorder::HandlePIE(bool bIsSimulating)
 {
 	ULevelSequence* FinishedAsset = GetSequence();
diff --git a/Engine/Plugins/VirtualProduction/Takes/Source/TakeRecorder/Public/Recorder/TakeRecorder.h b/Engine/Plugins/VirtualProduction/Takes/Source/TakeRecorder/Public/Recorder/TakeRecorder.h
--- a/Engine/Plugins/VirtualProduction/Takes/Source/TakeRecorder/Public/Recorder/TakeRecorder.h
+++ b/Engine/Plugins/VirtualProduction/Takes/Source/TakeRecorder/Public/Recorder/TakeRecorder.h
@@ -37,6 +37,7 @@ DECLARE_MULTICAST_DELEGATE_OneParam(FOnTakeRecordingStarted, UTakeRecorder*);
 DECLARE_MULTICAST_DELEGATE_OneParam(FOnTakeRecordingStopped, UTakeRecorder*);
 DECLARE_MULTICAST_DELEGATE_OneParam(FOnTakeRecordingFinished, UTakeRecorder*);
 DECLARE_MULTICAST_DELEGATE_OneParam(FOnTakeRecordingCancelled, UTakeRecorder*);
+DECLARE_MULTICAST_DELEGATE_TwoParams(FOnStartPlayFrameModified, UTakeRecorder*, const FFrameNumber& StartFrame);
 
 DECLARE_DELEGATE_RetVal_OneParam(FTakeRecorderParameters, FTakeRecorderParameterDelegate,const FTakeRecorderParameters&);
 
@@ -160,6 +161,11 @@ public:
 	 */
 	FOnTakeRecordingCancelled& OnRecordingCancelled();
 
+	/**
+	 * Retrieve a multi-cast delegate that is triggered when a delta time has been applied to the Movie Scene
+	 */
+	FOnStartPlayFrameModified& OnStartPlayFrameModified();
+
 private:
 
 	/**
@@ -279,6 +285,9 @@ private:
 	/** Triggered when this recorder is cancelled */
 	FOnTakeRecordingCancelled OnRecordingCancelledEvent;
 
+	/** Triggered when the movie scene is adjusted. */
+	FOnStartPlayFrameModified OnFrameModifiedEvent;
+
 	/** Sequencer ptr that controls playback of the desination asset during the recording */
 	TWeakPtr<ISequencer> WeakSequencer;
 
