Add new CVar to delay applying transactions while users is an editable text box.

diff --git a/Engine/Plugins/Developer/Concert/ConcertSync/ConcertSyncClient/Source/ConcertSyncClient/Private/ConcertClientWorkspace.cpp b/Engine/Plugins/Developer/Concert/ConcertSync/ConcertSyncClient/Source/ConcertSyncClient/Private/ConcertClientWorkspace.cpp
--- a/Engine/Plugins/Developer/Concert/ConcertSync/ConcertSyncClient/Source/ConcertSyncClient/Private/ConcertClientWorkspace.cpp
+++ b/Engine/Plugins/Developer/Concert/ConcertSync/ConcertSyncClient/Source/ConcertSyncClient/Private/ConcertClientWorkspace.cpp
@@ -834,7 +834,7 @@ void FConcertClientWorkspace::OnEndFrame()
 		FConcertSlowTaskStackWorkaround::Get().PopTask(MoveTemp(InitialSyncSlowTask));
 	}
 
-	if (bHasSyncedWorkspace && CanProcessPendingPackages())
+	if (bHasSyncedWorkspace && CanProcessPendingPackages() && !ConcertSyncClientUtil::UserIsEditing())
 	{
 		if (PackageManager)
 		{
diff --git a/Engine/Plugins/Developer/Concert/ConcertSync/ConcertSyncClient/Source/ConcertSyncClient/Private/ConcertSyncClientUtil.cpp b/Engine/Plugins/Developer/Concert/ConcertSync/ConcertSyncClient/Source/ConcertSyncClient/Private/ConcertSyncClientUtil.cpp
--- a/Engine/Plugins/Developer/Concert/ConcertSync/ConcertSyncClient/Source/ConcertSyncClient/Private/ConcertSyncClientUtil.cpp
+++ b/Engine/Plugins/Developer/Concert/ConcertSync/ConcertSyncClient/Source/ConcertSyncClient/Private/ConcertSyncClientUtil.cpp
@@ -44,6 +44,31 @@
 namespace ConcertSyncClientUtil
 {
 
+static TAutoConsoleVariable<int32> CVarDelayApplyingTransactionsWhileEditing(
+	TEXT("Concert.DelayTransactionsWhileEditing"), 0,
+	TEXT("Focus is lost by the editor when a transaction is applied. This variable suspends applying a transaction until the user has removed focus on editable UI."));
+
+bool UserIsEditing()
+{
+#if WITH_EDITOR
+	if (CVarDelayApplyingTransactionsWhileEditing.GetValueOnAnyThread() > 0)
+	{
+		static FName SEditableTextType(TEXT("SEditableText"));
+		static FName SMultiLineEditableTextType(TEXT("SMultiLineEditableText"));
+
+		bool bEditable = false;
+		FSlateApplication::Get().ForEachUser([&bEditable](FSlateUser& User) {
+			TSharedPtr<SWidget> FocusedWidget = User.GetFocusedWidget();
+
+			bool bCanEdit = FocusedWidget && (FocusedWidget->GetType() == SEditableTextType || FocusedWidget->GetType() == SMultiLineEditableTextType);
+			bEditable |= bCanEdit;
+		});
+		return bEditable;
+	}
+#endif
+	return false;
+}
+
 bool CanPerformBlockingAction(const bool bBlockDuringInteraction)
 {
 	// GUndo is a crude check to make sure that we don't try and apply other transactions while the local user is making a change
diff --git a/Engine/Plugins/Developer/Concert/ConcertSync/ConcertSyncClient/Source/ConcertSyncClient/Private/ConcertSyncClientUtil.h b/Engine/Plugins/Developer/Concert/ConcertSync/ConcertSyncClient/Source/ConcertSyncClient/Private/ConcertSyncClientUtil.h
--- a/Engine/Plugins/Developer/Concert/ConcertSync/ConcertSyncClient/Source/ConcertSyncClient/Private/ConcertSyncClientUtil.h
+++ b/Engine/Plugins/Developer/Concert/ConcertSync/ConcertSyncClient/Source/ConcertSyncClient/Private/ConcertSyncClientUtil.h
@@ -62,6 +62,8 @@ namespace ConcertSyncClientUtil
 		EGetObjectResultFlags Flags;
 	};
 
+	bool UserIsEditing();
+
 	bool CanPerformBlockingAction(const bool bBlockDuringInteraction = true);
 
 	void UpdatePendingKillState(UObject* InObj, const bool bIsPendingKill);
