Make it easier to override the reflect visibility in game.

diff --git a/Engine/Plugins/Developer/Concert/ConcertSync/ConcertSyncClient/Source/ConcertSyncClient/Private/ConcertClientWorkspace.cpp b/Engine/Plugins/Developer/Concert/ConcertSync/ConcertSyncClient/Source/ConcertSyncClient/Private/ConcertClientWorkspace.cpp
--- a/Engine/Plugins/Developer/Concert/ConcertSync/ConcertSyncClient/Source/ConcertSyncClient/Private/ConcertClientWorkspace.cpp
+++ b/Engine/Plugins/Developer/Concert/ConcertSync/ConcertSyncClient/Source/ConcertSyncClient/Private/ConcertClientWorkspace.cpp
@@ -119,15 +119,24 @@ private:
 	TArray<TUniquePtr<FScopedSlowTask>> ExtendedTaskLife;
 };
 
-void SetReflectEditorLevelVisibilityWithGame(bool InValue)
+void SetReflectEditorLevelVisibilityWithGame(bool bInValue, bool bReset = false)
 {
 #if WITH_EDITOR
+	static bool bHasBeenSet = false;
+
+	if (bReset)
+	{
+		bHasBeenSet = false;
+	}
+
 	// Detail mode was modified, so store in the CVar
 	static IConsoleVariable* CVar = IConsoleManager::Get().FindConsoleVariable(TEXT("Editor.ReflectEditorLevelVisibilityWithGame"));
-	int32 ValAsInt = !!InValue;
-	if (CVar->GetInt() != ValAsInt)
+	int32 ValAsInt = !!bInValue;
+	bool bCanSet = !bHasBeenSet && CVar->GetInt() != ValAsInt;
+	if (GEditor && bCanSet)
 	{
 		CVar->Set(ValAsInt);
+		bHasBeenSet = bReset ? false : true;
 	}
 #endif
 }
@@ -639,7 +648,7 @@ void FConcertClientWorkspace::HandleConnectionChanged(IConcertClientSession& InS
 		bHasSyncedWorkspace = false;
 		bFinalizeWorkspaceSyncRequested = false;
 		FConcertSlowTaskStackWorkaround::Get().PopTask(MoveTemp(InitialSyncSlowTask));
-		SetReflectEditorLevelVisibilityWithGame(false);
+		SetReflectEditorLevelVisibilityWithGame(false, true);
 	}
 }
 
