Do not permit MarkForAdd while in sandbox mode.

diff --git a/Engine/Plugins/Developer/Concert/ConcertSync/ConcertSyncClient/Source/ConcertSyncClient/Private/ConcertSourceControlProxy.cpp b/Engine/Plugins/Developer/Concert/ConcertSync/ConcertSyncClient/Source/ConcertSyncClient/Private/ConcertSourceControlProxy.cpp
--- a/Engine/Plugins/Developer/Concert/ConcertSync/ConcertSyncClient/Source/ConcertSyncClient/Private/ConcertSourceControlProxy.cpp
+++ b/Engine/Plugins/Developer/Concert/ConcertSync/ConcertSyncClient/Source/ConcertSyncClient/Private/ConcertSourceControlProxy.cpp
@@ -436,7 +436,13 @@ ECommandResult::Type FConcertSourceControlProxy::Execute(const FSourceControlOpe
 {
 	if (ActualProvider)
 	{
-		return ActualProvider->Execute(InOperation, InChangelist, InFiles, InConcurrency, InOperationCompleteDelegate);
+		// We do not want to allow a P4 add when we operating within a sandbox. Perforce is the only SCS that will allow you to add a file that does
+		// not yet exist.
+		//
+		if (!IFileManager::Get().IsSandboxEnabled() || InOperation->GetName() != FName(TEXT("MarkForAdd")))
+		{
+			return ActualProvider->Execute(InOperation, InChangelist, InFiles, InConcurrency, InOperationCompleteDelegate);
+		}
 	}
 
 	return ECommandResult::Failed;
