Added an override iwth a FConcertSyncWorldRemapper parameter to the IConcertClientTransactionBridge::ApplyRemoteTransaction prototype. This allows remapping to occur internally if needed.

diff --git a/Engine/Plugins/Developer/Concert/ConcertSync/ConcertSyncClient/Source/ConcertSyncClient/Private/ConcertClientTransactionBridge.cpp b/Engine/Plugins/Developer/Concert/ConcertSync/ConcertSyncClient/Source/ConcertSyncClient/Private/ConcertClientTransactionBridge.cpp
--- a/Engine/Plugins/Developer/Concert/ConcertSync/ConcertSyncClient/Source/ConcertSyncClient/Private/ConcertClientTransactionBridge.cpp
+++ b/Engine/Plugins/Developer/Concert/ConcertSync/ConcertSyncClient/Source/ConcertSyncClient/Private/ConcertClientTransactionBridge.cpp
@@ -3,7 +3,6 @@
 #include "ConcertClientTransactionBridge.h"
 #include "ConcertLogGlobal.h"
 #include "ConcertSyncSettings.h"
-#include "ConcertSyncArchives.h"
 #include "ConcertSyncClientUtil.h"
 
 #include "ConcertTransactionEvents.h"
@@ -178,7 +177,7 @@ struct FEditorTransactionNotification
 };
 #endif
 
-void ProcessTransactionEvent(const FConcertTransactionEventBase& InEvent, const FConcertSessionVersionInfo* InVersionInfo, const TArray<FName>& InPackagesToProcess, const FConcertLocalIdentifierTable* InLocalIdentifierTablePtr, const bool bIsSnapshot)
+void ProcessTransactionEvent(const FConcertTransactionEventBase& InEvent, const FConcertSessionVersionInfo* InVersionInfo, const TArray<FName>& InPackagesToProcess, const FConcertLocalIdentifierTable* InLocalIdentifierTablePtr, const bool bIsSnapshot, const FConcertSyncWorldRemapper& WorldRemapper)
 {
 	// Transactions are applied in multiple-phases...
 	//	1) Find or create all objects in the transaction (to handle object-interdependencies in the serialized data)
@@ -263,7 +262,7 @@ void ProcessTransactionEvent(const FConcertTransactionEventBase& InEvent, const
 		TSharedPtr<ITransactionObjectAnnotation>& TransactionAnnotation = TransactionAnnotations[ObjectIndex];
 		if (ObjectUpdate.SerializedAnnotationData.Num() > 0)
 		{
-			FConcertSyncObjectReader AnnotationReader(InLocalIdentifierTablePtr, FConcertSyncWorldRemapper(), InVersionInfo, TransactionObject, ObjectUpdate.SerializedAnnotationData);
+			FConcertSyncObjectReader AnnotationReader(InLocalIdentifierTablePtr, WorldRemapper, InVersionInfo, TransactionObject, ObjectUpdate.SerializedAnnotationData);
 			TransactionAnnotation = TransactionObject->CreateAndRestoreTransactionAnnotation(AnnotationReader);
 			UE_CLOG(!TransactionAnnotation.IsValid(), LogConcert, Warning, TEXT("Object '%s' had transaction annotation data that failed to restore!"), *TransactionObject->GetPathName());
 		}
@@ -315,7 +314,7 @@ void ProcessTransactionEvent(const FConcertTransactionEventBase& InEvent, const
 		// Apply the new data
 		if (ObjectUpdate.ObjectData.SerializedData.Num() > 0)
 		{
-			FConcertSyncObjectReader ObjectReader(InLocalIdentifierTablePtr, FConcertSyncWorldRemapper(), InVersionInfo, TransactionObject, ObjectUpdate.ObjectData.SerializedData);
+			FConcertSyncObjectReader ObjectReader(InLocalIdentifierTablePtr, WorldRemapper, InVersionInfo, TransactionObject, ObjectUpdate.ObjectData.SerializedData);
 			ObjectReader.SerializeObject(TransactionObject);
 		}
 		else
@@ -325,7 +324,7 @@ void ProcessTransactionEvent(const FConcertTransactionEventBase& InEvent, const
 				FProperty* TransactionProp = FindFProperty<FProperty>(TransactionObject->GetClass(), PropertyData.PropertyName);
 				if (TransactionProp)
 				{
-					FConcertSyncObjectReader ObjectReader(InLocalIdentifierTablePtr, FConcertSyncWorldRemapper(), InVersionInfo, TransactionObject, PropertyData.SerializedData);
+					FConcertSyncObjectReader ObjectReader(InLocalIdentifierTablePtr, WorldRemapper, InVersionInfo, TransactionObject, PropertyData.SerializedData);
 					ObjectReader.SerializeProperty(TransactionProp, TransactionObject);
 				}
 			}
@@ -472,7 +471,12 @@ FOnApplyTransaction& FConcertClientTransactionBridge::OnApplyTransaction()
 
 void FConcertClientTransactionBridge::ApplyRemoteTransaction(const FConcertTransactionEventBase& InEvent, const FConcertSessionVersionInfo* InVersionInfo, const TArray<FName>& InPackagesToProcess, const FConcertLocalIdentifierTable* InLocalIdentifierTablePtr, const bool bIsSnapshot)
 {
-	ConcertClientTransactionBridgeUtil::ProcessTransactionEvent(InEvent, InVersionInfo, InPackagesToProcess, InLocalIdentifierTablePtr, bIsSnapshot);
+	ConcertClientTransactionBridgeUtil::ProcessTransactionEvent(InEvent, InVersionInfo, InPackagesToProcess, InLocalIdentifierTablePtr, bIsSnapshot, FConcertSyncWorldRemapper());
+}
+
+void FConcertClientTransactionBridge::ApplyRemoteTransaction(const FConcertTransactionEventBase& InEvent, const FConcertSessionVersionInfo* InVersionInfo, const TArray<FName>& InPackagesToProcess, const FConcertLocalIdentifierTable* InLocalIdentifierTablePtr, const bool bIsSnapshot, const FConcertSyncWorldRemapper& ConcertSyncWorldRemapper)
+{
+	ConcertClientTransactionBridgeUtil::ProcessTransactionEvent(InEvent, InVersionInfo, InPackagesToProcess, InLocalIdentifierTablePtr, bIsSnapshot, ConcertSyncWorldRemapper);
 }
 
 bool& FConcertClientTransactionBridge::GetIgnoreLocalTransactionsRef()
diff --git a/Engine/Plugins/Developer/Concert/ConcertSync/ConcertSyncClient/Source/ConcertSyncClient/Private/ConcertClientTransactionBridge.h b/Engine/Plugins/Developer/Concert/ConcertSync/ConcertSyncClient/Source/ConcertSyncClient/Private/ConcertClientTransactionBridge.h
--- a/Engine/Plugins/Developer/Concert/ConcertSync/ConcertSyncClient/Source/ConcertSyncClient/Private/ConcertClientTransactionBridge.h
+++ b/Engine/Plugins/Developer/Concert/ConcertSync/ConcertSyncClient/Source/ConcertSyncClient/Private/ConcertClientTransactionBridge.h
@@ -4,6 +4,7 @@
 
 #include "IConcertClientTransactionBridge.h"
 #include "Misc/ITransaction.h"
+#include "ConcertSyncArchives.h"
 
 class CONCERTSYNCCLIENT_API FConcertClientTransactionBridge : public IConcertClientTransactionBridge
 {
@@ -19,6 +20,7 @@ public:
 	virtual FOnApplyTransaction& OnApplyTransaction() override;
 
 	virtual void ApplyRemoteTransaction(const FConcertTransactionEventBase& InEvent, const FConcertSessionVersionInfo* InVersionInfo, const TArray<FName>& InPackagesToProcess, const FConcertLocalIdentifierTable* InLocalIdentifierTablePtr, const bool bIsSnapshot) override;
+	virtual void ApplyRemoteTransaction(const FConcertTransactionEventBase& InEvent, const FConcertSessionVersionInfo* InVersionInfo, const TArray<FName>& InPackagesToProcess, const FConcertLocalIdentifierTable* InLocalIdentifierTablePtr, const bool bIsSnapshot, const FConcertSyncWorldRemapper& ConcertSyncWorldRemapper) override;
 	virtual bool& GetIgnoreLocalTransactionsRef() override;
 
 	virtual void RegisterTransactionFilter(FName FilterName, FTransactionFilterDelegate FilterHandle) override;
@@ -79,4 +81,6 @@ private:
 
 	/** Include non-cooked properties in object serialization */
 	bool bIncludeEditorOnlyProperties;
+
+	FConcertSyncWorldRemapper WorldRemapper = FConcertSyncWorldRemapper();
 };
diff --git a/Engine/Plugins/Developer/Concert/ConcertSync/ConcertSyncClient/Source/ConcertSyncClient/Public/IConcertClientTransactionBridge.h b/Engine/Plugins/Developer/Concert/ConcertSync/ConcertSyncClient/Source/ConcertSyncClient/Public/IConcertClientTransactionBridge.h
--- a/Engine/Plugins/Developer/Concert/ConcertSync/ConcertSyncClient/Source/ConcertSyncClient/Public/IConcertClientTransactionBridge.h
+++ b/Engine/Plugins/Developer/Concert/ConcertSync/ConcertSyncClient/Source/ConcertSyncClient/Public/IConcertClientTransactionBridge.h
@@ -116,6 +116,17 @@ public:
 	 */
 	virtual void ApplyRemoteTransaction(const FConcertTransactionEventBase& InEvent, const FConcertSessionVersionInfo* InVersionInfo, const TArray<FName>& InPackagesToProcess, const FConcertLocalIdentifierTable* InLocalIdentifierTablePtr, const bool bIsSnapshot) = 0;
 
+	/**
+	 * Apply a remote transaction event to this local instance.
+	 * @param InEvent					The event to apply.
+	 * @param InVersionInfo				The version information for the serialized data in the event, or null if the event should be serialized using the compiled in version info.
+	 * @param InPackagesToProcess		The list of packages to apply changes for, or an empty array to apply all changes.
+	 * @param InLocalIdentifierTablePtr The local identifier table for the event data (if any).
+	 * @param bIsSnapshot				True if this transaction event was a snapshot rather than a finalized transaction.
+	 * @param ConcertSyncWorldRemapper	Remapper to use in the case the current world is different from the one sent in the transation.
+	 */
+	virtual void ApplyRemoteTransaction(const FConcertTransactionEventBase& InEvent, const FConcertSessionVersionInfo* InVersionInfo, const TArray<FName>& InPackagesToProcess, const FConcertLocalIdentifierTable* InLocalIdentifierTablePtr, const bool bIsSnapshot, const class FConcertSyncWorldRemapper& ConcertSyncWorldRemapper) = 0;
+
 	/** Callback to register delegate for handling transaction events */
 	virtual void RegisterTransactionFilter(FName FilterName, FTransactionFilterDelegate FilterHandle) = 0;
 
