If a package is saved asynchronously it will not be reported to exist until later. We push the package onto the list that would trigger the handler but only call the delegate if it exists after we flush asynchronous package writes.

diff --git a/Engine/Plugins/Developer/Concert/ConcertSync/ConcertSyncClient/Source/ConcertSyncClient/Private/ConcertClientPackageBridge.cpp b/Engine/Plugins/Developer/Concert/ConcertSync/ConcertSyncClient/Source/ConcertSyncClient/Private/ConcertClientPackageBridge.cpp
--- a/Engine/Plugins/Developer/Concert/ConcertSync/ConcertSyncClient/Source/ConcertSyncClient/Private/ConcertClientPackageBridge.cpp
+++ b/Engine/Plugins/Developer/Concert/ConcertSync/ConcertSyncClient/Source/ConcertSyncClient/Private/ConcertClientPackageBridge.cpp
@@ -178,7 +178,11 @@ void FConcertClientPackageBridge::OnEndFrame()
 	UPackage::WaitForAsyncFileWrites();
 	for (const FConcertPackageInfoTuple& PackageInfoTuple : PendingPackageInfos)
 	{
-		OnLocalPackageEventDelegate.Broadcast(PackageInfoTuple.Get<0>(), PackageInfoTuple.Get<1>());
+		const FString& PackageFilename = PackageInfoTuple.Get<1>();
+		if (IFileManager::Get().FileExists(*PackageFilename))
+		{
+			OnLocalPackageEventDelegate.Broadcast(PackageInfoTuple.Get<0>(), PackageFilename);
+		}
 	}
 	PendingPackageInfos.Reset();
 }
@@ -207,16 +211,13 @@ void FConcertClientPackageBridge::HandlePackageSaved(const FString& PackageFilen
 	FName NewPackageName;
 	PackagesBeingRenamed.RemoveAndCopyValue(Package->GetFName(), NewPackageName);
 
-	if (IFileManager::Get().FileExists(*PackageFilename))
-	{
-		FConcertPackageInfo PackageInfo;
-		ConcertSyncClientUtil::FillPackageInfo(Package, nullptr, NewPackageName.IsNone() ? EConcertPackageUpdateType::Saved : EConcertPackageUpdateType::Renamed, PackageInfo);
-		PackageInfo.NewPackageName = NewPackageName;
-		PackageInfo.bPreSave = false;
-		PackageInfo.bAutoSave = GEngine->IsAutosaving();
+	FConcertPackageInfo PackageInfo;
+	ConcertSyncClientUtil::FillPackageInfo(Package, nullptr, NewPackageName.IsNone() ? EConcertPackageUpdateType::Saved : EConcertPackageUpdateType::Renamed, PackageInfo);
+	PackageInfo.NewPackageName = NewPackageName;
+	PackageInfo.bPreSave = false;
+	PackageInfo.bAutoSave = GEngine->IsAutosaving();
 
-		PendingPackageInfos.Add(FConcertPackageInfoTuple(MoveTemp(PackageInfo), PackageFilename));
-	}
+	PendingPackageInfos.Add(FConcertPackageInfoTuple(MoveTemp(PackageInfo), PackageFilename));
 
 	UE_LOG(LogConcert, Verbose, TEXT("Asset Saved: %s"), *Package->GetName());
 }
