Add unbounded play to multi-user sequencer client and remove code on the player that was improperly adjust range.  This is now handled directly during a transport event update.

diff --git a/Engine/Plugins/Developer/Concert/ConcertSync/ConcertSyncClient/Source/ConcertSyncClient/Private/ConcertClientSequencerManager.cpp b/Engine/Plugins/Developer/Concert/ConcertSync/ConcertSyncClient/Source/ConcertSyncClient/Private/ConcertClientSequencerManager.cpp
--- a/Engine/Plugins/Developer/Concert/ConcertSync/ConcertSyncClient/Source/ConcertSyncClient/Private/ConcertClientSequencerManager.cpp
+++ b/Engine/Plugins/Developer/Concert/ConcertSync/ConcertSyncClient/Source/ConcertSyncClient/Private/ConcertClientSequencerManager.cpp
@@ -4,6 +4,7 @@
 #include "ConcertSequencerMessages.h"
 #include "Delegates/IDelegateInstance.h"
 #include "HAL/IConsoleManager.h"
+#include "Misc/AssertionMacros.h"
 #include "Misc/QualifiedFrameTime.h"
 #include "Modules/ModuleManager.h"
 #include "Misc/CoreDelegates.h"
@@ -331,6 +332,10 @@ void FConcertClientSequencerManager::OnSequencerTimeChanged(TWeakPtr<ISequencer>
 			SequencerStateEvent.State.Time					= Sequencer->GetGlobalTime();
 			SequencerStateEvent.State.PlayerStatus			= (EConcertMovieScenePlayerStatus)Sequencer->GetPlaybackStatus();
 			SequencerStateEvent.State.PlaybackSpeed			= Sequencer->GetPlaybackSpeed();
+
+			UMovieScene* MovieScene = Sequence->GetMovieScene();
+			check(MovieScene);
+			SequencerStateEvent.State.PlaybackRange			= MovieScene->GetPlaybackRange();
 			SequencerState = SequencerStateEvent.State;
 
 			// Send to client and server
@@ -458,45 +463,16 @@ void FConcertClientSequencerManager::CreateNewSequencePlayerIfNotExists(const FS
 		return;
 	}
 
-	FDelegateHandle Handle;
 	FMovieSceneSequencePlaybackSettings PlaybackSettings;
 	// Sequencer behaves differently to Player.
 	// Sequencer pauses at the last frame and Player Stops and goes to the first frame unless we set this flag.
 	PlaybackSettings.bPauseAtEnd = true;
 
 	// This call will initialeze LevelSequenceActor as an output parameter.
-	ULevelSequencePlayer* Player =ULevelSequencePlayer::CreateLevelSequencePlayer(CurrentWorld->PersistentLevel, Sequence, PlaybackSettings, LevelSequenceActor);
+	ULevelSequencePlayer* Player = ULevelSequencePlayer::CreateLevelSequencePlayer(CurrentWorld->PersistentLevel, Sequence, PlaybackSettings, LevelSequenceActor);
 	check(Player);
 
-	UMovieScene* Scene = Sequence->GetMovieScene();
-	check(Scene != nullptr);
-
-	Handle = Scene->OnSignatureChanged().AddLambda([SceneObj = TWeakObjectPtr<UMovieScene>(Scene),
-													LevelActorObj = TWeakObjectPtr<ALevelSequenceActor>(LevelSequenceActor)]() {
-		if (LevelActorObj.IsValid() && SceneObj.IsValid())
-		{
-			const TRange<FFrameNumber> PlayRange = SceneObj->GetPlaybackRange();
-			const FFrameRate TickResolution = SceneObj->GetTickResolution();
-			const FFrameRate DisplayRate = SceneObj->GetDisplayRate();
-
-			const FFrameNumber SrcStartFrame = UE::MovieScene::DiscreteInclusiveLower(PlayRange);
-			const FFrameNumber SrcEndFrame   = UE::MovieScene::DiscreteExclusiveUpper(PlayRange);
-
-			const FFrameTime EndingTime = ConvertFrameTime(SrcEndFrame, TickResolution, DisplayRate);
-
-			const FFrameNumber StartingFrame = ConvertFrameTime(SrcStartFrame, TickResolution, DisplayRate).FloorToFrame();
-			const FFrameNumber EndingFrame   = EndingTime.FloorToFrame();
-
-			const int32 CurrentDuration = LevelActorObj->GetSequencePlayer()->GetFrameDuration();
-			const FQualifiedFrameTime QualifiedStartTime = LevelActorObj->GetSequencePlayer()->GetStartTime();
-			const int32 NewDuration = (EndingFrame - StartingFrame).Value;
-			if (CurrentDuration != NewDuration || QualifiedStartTime.Time.GetFrame() != StartingFrame)
-			{
-				LevelActorObj->GetSequencePlayer()->SetFrameRange(StartingFrame.Value, NewDuration, EndingTime.GetSubFrame());
-			}
-		}
-	});
-	SequencePlayers.Add(*SequenceObjectPath, {LevelSequenceActor, MoveTemp(Handle)});
+	SequencePlayers.Add(*SequenceObjectPath, {LevelSequenceActor});
 }
 
 bool FConcertClientSequencerManager::CanClose(const FConcertSequencerCloseEvent& InEvent) const
@@ -512,15 +488,6 @@ void FConcertClientSequencerManager::DestroyPlayer(FSequencePlayer& Player, ALev
 	{
 		LevelSequenceActor->SequencePlayer->Stop();
 
-		ULevelSequence* Sequence = LevelSequenceActor->GetSequence();
-		if (Sequence)
-		{
-			UMovieScene* Scene = Sequence->GetMovieScene();
-			if (Scene && Player.SignatureChangedHandle.IsValid())
-			{
-				Scene->OnSignatureChanged().Remove(Player.SignatureChangedHandle);
-			}
-		}
 		LevelSequenceActor->SetSequence(nullptr);
 		LevelSequenceActor->SequencePlayer = nullptr;
 		LevelSequenceActor->Destroy(false, false);
@@ -702,6 +669,10 @@ void FConcertClientSequencerManager::ApplyEventToSequencers(const FConcertSequen
 		FFrameRate SequenceRate = Sequencer->GetRootTickResolution();
 		FFrameTime IncomingTime = EventState.Time.ConvertTo(SequenceRate);
 
+		UMovieScene* MovieScene = Sequencer->GetRootMovieSceneSequence()->GetMovieScene();
+		check(MovieScene);
+		MovieScene->SetPlaybackRange(EventState.PlaybackRange, false);
+
 		// If the event is coming from a sequencer that is playing back, we are a slave to its updates until it stops
 		// We also apply any latency compensation when playing back
 		if (EventState.PlayerStatus == EConcertMovieScenePlayerStatus::Playing)
@@ -768,6 +739,42 @@ void FConcertClientSequencerManager::ApplyEventToSequencers(const FConcertSequen
 	}
 }
 
+namespace UE::Private::ConcertClientSequencerManager
+{
+
+int32 GetDuration(const FFrameNumberRange& PlayRange)
+{
+	const TRangeBound<FFrameNumber>& UpperBound = PlayRange.GetUpperBound();
+	if (UpperBound.IsClosed())
+	{
+		const FFrameNumber SrcStartFrame = UE::MovieScene::DiscreteInclusiveLower(PlayRange);
+		const FFrameNumber SrcEndFrame   = UE::MovieScene::DiscreteExclusiveUpper(PlayRange);
+		return (SrcEndFrame - SrcStartFrame).Value;
+	}
+	return TNumericLimits<int32>::Max() - 1;
+}
+
+void ApplyPlayRangeToPlayer(ULevelSequencePlayer* Player, const FFrameNumberRange& PlayRange)
+{
+	check(Player);
+
+	if (ensureMsgf(PlayRange.GetLowerBound().IsClosed(), TEXT("PlayRange lower bound is open, which is not supported.")))
+	{
+		const FFrameNumber SrcStartFrame = UE::MovieScene::DiscreteInclusiveLower(PlayRange);
+
+		const int32 CurrentDuration = Player->GetFrameDuration();
+		const int32 NewDuration = GetDuration(PlayRange);
+
+		const FQualifiedFrameTime QualifiedStartTime = Player->GetStartTime();
+		if (CurrentDuration != NewDuration || QualifiedStartTime.Time.GetFrame() != SrcStartFrame)
+		{
+			Player->SetFrameRange(SrcStartFrame.Value, NewDuration);
+		}
+	}
+}
+
+}
+
 
 void FConcertClientSequencerManager::ApplyEventToPlayers(const FConcertSequencerState& EventState)
 {
@@ -787,6 +794,8 @@ void FConcertClientSequencerManager::ApplyEventToPlayers(const FConcertSequencer
 		FFrameRate SequenceRate = Player->GetFrameRate();
 		FFrameTime IncomingTime = EventState.Time.ConvertTo(SequenceRate);
 
+		UE::Private::ConcertClientSequencerManager::ApplyPlayRangeToPlayer(Player, EventState.PlaybackRange);
+
 		// If the event is coming from a sequencer that is playing back, we are a slave to its updates until it stops
 		// We also apply any latency compensation when playing back
 		if (EventState.PlayerStatus == EConcertMovieScenePlayerStatus::Playing)
diff --git a/Engine/Plugins/Developer/Concert/ConcertSync/ConcertSyncClient/Source/ConcertSyncClient/Private/ConcertClientSequencerManager.h b/Engine/Plugins/Developer/Concert/ConcertSync/ConcertSyncClient/Source/ConcertSyncClient/Private/ConcertClientSequencerManager.h
--- a/Engine/Plugins/Developer/Concert/ConcertSync/ConcertSyncClient/Source/ConcertSyncClient/Private/ConcertClientSequencerManager.h
+++ b/Engine/Plugins/Developer/Concert/ConcertSync/ConcertSyncClient/Source/ConcertSyncClient/Private/ConcertClientSequencerManager.h
@@ -23,7 +23,6 @@ struct FSequencerInitParams;
 struct FSequencePlayer
 {
 	TWeakObjectPtr<ALevelSequenceActor> Actor;
-	FDelegateHandle	SignatureChangedHandle;
 };
 
 /**
@@ -301,9 +300,6 @@ private:
 	/** Indicates if this sequence player can be closed. */
 	bool CanClose(const FConcertSequencerCloseEvent& InEvent) const;
 
-	/** Called by the MovieScene object when the property changes of the object. */
-	void OnPlayerSignatureChanged();
-
 	/** Pointer to the sync client that owns us. */
 	IConcertSyncClient* OwnerSyncClient;
 
diff --git a/Engine/Plugins/Developer/Concert/ConcertSync/ConcertSyncCore/Source/ConcertSyncCore/Public/ConcertSequencerMessages.h b/Engine/Plugins/Developer/Concert/ConcertSync/ConcertSyncCore/Source/ConcertSyncCore/Public/ConcertSequencerMessages.h
--- a/Engine/Plugins/Developer/Concert/ConcertSync/ConcertSyncCore/Source/ConcertSyncCore/Public/ConcertSequencerMessages.h
+++ b/Engine/Plugins/Developer/Concert/ConcertSync/ConcertSyncCore/Source/ConcertSyncCore/Public/ConcertSequencerMessages.h
@@ -6,6 +6,7 @@
 #include "UObject/Object.h"
 #include "UObject/Class.h"
 #include "Misc/QualifiedFrameTime.h"
+#include "Math/Range.h"
 #include "ConcertSequencerMessages.generated.h"
 
 
@@ -42,6 +43,9 @@ struct FConcertSequencerState
 	UPROPERTY()
 	EConcertMovieScenePlayerStatus PlayerStatus;
 
+	UPROPERTY()
+	FFrameNumberRange PlaybackRange;
+
 	/** The current playback speed */
 	UPROPERTY()
 	float PlaybackSpeed;
