Create a flush package loading path that does not call ResetLoaders().  It is not necessary on discard sandbox due to the GC of packages anyway.

diff --git a/Engine/Plugins/Developer/Concert/ConcertSync/ConcertSyncClient/Source/ConcertSyncClient/Private/ConcertSandboxPlatformFile.cpp b/Engine/Plugins/Developer/Concert/ConcertSync/ConcertSyncClient/Source/ConcertSyncClient/Private/ConcertSandboxPlatformFile.cpp
--- a/Engine/Plugins/Developer/Concert/ConcertSync/ConcertSyncClient/Source/ConcertSyncClient/Private/ConcertSandboxPlatformFile.cpp
+++ b/Engine/Plugins/Developer/Concert/ConcertSync/ConcertSyncClient/Source/ConcertSyncClient/Private/ConcertSandboxPlatformFile.cpp
@@ -41,12 +41,12 @@ FString GetContentFolderName(const FString& InContentPath)
 	return ContentFolderName;
 }
 
-bool FlushPackageFile(const FString& InFilename, FName* OutPackageName = nullptr)
+bool FlushPackageFile(const FString& InFilename, FName* OutPackageName = nullptr, bool bForceLoad = true)
 {
 	FString PackageName;
 	if (FPackageName::TryConvertFilenameToLongPackageName(InFilename, PackageName))
 	{
-		ConcertSyncClientUtil::FlushPackageLoading(PackageName);
+		ConcertSyncClientUtil::FlushPackageLoading(PackageName, bForceLoad);
 		if (OutPackageName)
 		{
 			*OutPackageName = *PackageName;
@@ -954,7 +954,7 @@ void FConcertSandboxPlatformFile::DiscardSandbox(TArray<FName>& OutPackagesPendi
 
 					// If this file maps to a package then we need to flush its linker so that we can remove the file from the sandbox
 					FName PackageName;
-					ConcertSandboxPlatformFileUtil::FlushPackageFile(RemappedFilePath.GetNonSandboxPath(), &PackageName);
+					ConcertSandboxPlatformFileUtil::FlushPackageFile(RemappedFilePath.GetNonSandboxPath(), &PackageName, false);
 
 					if (LowerLevel->FileExists(*RemappedFilePath.GetNonSandboxPath()))
 					{
diff --git a/Engine/Plugins/Developer/Concert/ConcertSync/ConcertSyncClient/Source/ConcertSyncClient/Private/ConcertSyncClientUtil.cpp b/Engine/Plugins/Developer/Concert/ConcertSync/ConcertSyncClient/Source/ConcertSyncClient/Private/ConcertSyncClientUtil.cpp
--- a/Engine/Plugins/Developer/Concert/ConcertSync/ConcertSyncClient/Source/ConcertSyncClient/Private/ConcertSyncClientUtil.cpp
+++ b/Engine/Plugins/Developer/Concert/ConcertSync/ConcertSyncClient/Source/ConcertSyncClient/Private/ConcertSyncClientUtil.cpp
@@ -348,7 +348,7 @@ void FlushPackageLoading(const FName InPackageName)
 	FlushPackageLoading(InPackageName.ToString());
 }
 
-void FlushPackageLoading(const FString& InPackageName)
+void FlushPackageLoading(const FString& InPackageName, bool bForceBulkDataLoad)
 {
 	UPackage* ExistingPackage = FindPackage(nullptr, *InPackageName);
 	if (ExistingPackage)
@@ -358,7 +358,15 @@ void FlushPackageLoading(const FString& InPackageName)
 			FlushAsyncLoading();
 			ExistingPackage->FullyLoad();
 		}
-		ResetLoaders(ExistingPackage);
+
+		if (bForceBulkDataLoad)
+		{
+			ResetLoaders(ExistingPackage);
+		}
+		else if (ExistingPackage->GetLinker())
+		{
+			ExistingPackage->GetLinker()->Detach();
+		}
 	}
 }
 
diff --git a/Engine/Plugins/Developer/Concert/ConcertSync/ConcertSyncClient/Source/ConcertSyncClient/Private/ConcertSyncClientUtil.h b/Engine/Plugins/Developer/Concert/ConcertSync/ConcertSyncClient/Source/ConcertSyncClient/Private/ConcertSyncClientUtil.h
--- a/Engine/Plugins/Developer/Concert/ConcertSync/ConcertSyncClient/Source/ConcertSyncClient/Private/ConcertSyncClientUtil.h
+++ b/Engine/Plugins/Developer/Concert/ConcertSync/ConcertSyncClient/Source/ConcertSyncClient/Private/ConcertSyncClientUtil.h
@@ -88,7 +88,7 @@ namespace ConcertSyncClientUtil
 
 	void FlushPackageLoading(const FName InPackageName);
 
-	void FlushPackageLoading(const FString& InPackageName);
+	void FlushPackageLoading(const FString& InPackageName, bool bForceBulkDataLoad = true);
 
 	void HotReloadPackages(TArrayView<const FName> InPackageNames);
 
