Multi-user Sequencer: Fix for recent regression causing play range adjustment to be incorrectly set using value in tick resolution timebase instead of display rate. Also changes LogConcertSequencerSync not to compile out Verbose log statements, and clarifies some existing log messages.

diff --git a/Engine/Plugins/Developer/Concert/ConcertSync/ConcertSyncClient/Source/ConcertSyncClient/Private/ConcertClientSequencerManager.cpp b/Engine/Plugins/Developer/Concert/ConcertSync/ConcertSyncClient/Source/ConcertSyncClient/Private/ConcertClientSequencerManager.cpp
--- a/Engine/Plugins/Developer/Concert/ConcertSync/ConcertSyncClient/Source/ConcertSyncClient/Private/ConcertClientSequencerManager.cpp
+++ b/Engine/Plugins/Developer/Concert/ConcertSync/ConcertSyncClient/Source/ConcertSyncClient/Private/ConcertClientSequencerManager.cpp
@@ -32,7 +32,7 @@
 	#include "Editor.h"
 #endif
 
-DEFINE_LOG_CATEGORY_STATIC(LogConcertSequencerSync, Warning, Log)
+DEFINE_LOG_CATEGORY_STATIC(LogConcertSequencerSync, Warning, All)
 
 #if WITH_EDITOR
 
@@ -128,7 +128,7 @@ void FConcertClientSequencerManager::OnSequencerCreated(TSharedRef<ISequencer> I
 				FConcertSequencerOpenEvent OpenEvent;
 				OpenEvent.SequenceObjectPath = Sequence->GetPathName();
 
-				UE_LOG(LogConcertSequencerSync, Verbose, TEXT("OnsequencerCreated: %s"), *OpenEvent.SequenceObjectPath);
+				UE_LOG(LogConcertSequencerSync, Verbose, TEXT("Sending OpenEvent: %s"), *OpenEvent.SequenceObjectPath);
 				Session->SendCustomEvent(OpenEvent, Session->GetSessionServerEndpointId(), EConcertMessageFlags::ReliableOrdered);
 			}
 		}
@@ -257,7 +257,7 @@ void FConcertClientSequencerManager::OnSequencerClosed(TSharedRef<ISequencer> In
 		{
 			if (CanSendSequencerEvent(Sequence->GetPathName()))
 			{
-				UE_LOG(LogConcertSequencerSync, Verbose, TEXT("Sending close event %s"), *Sequence->GetPathName());
+				UE_LOG(LogConcertSequencerSync, Verbose, TEXT("Sending CloseEvent: %s"), *Sequence->GetPathName());
 
 				FConcertSequencerCloseEvent CloseEvent;
 				CloseEvent.bControllerClose = ClosingSequencer.PlaybackMode == EPlaybackMode::Controller; // this sequencer had control over the sequence playback
@@ -339,7 +339,8 @@ void FConcertClientSequencerManager::OnSequencerTimeChanged(TWeakPtr<ISequencer>
 			SequencerState = SequencerStateEvent.State;
 
 			// Send to client and server
-			UE_LOG(LogConcertSequencerSync, Verbose, TEXT("OnSequencerTimeChanged: %s, at frame: %d"), *SequencerStateEvent.State.SequenceObjectPath, SequencerStateEvent.State.Time.Time.FrameNumber.Value);
+			UE_LOG(LogConcertSequencerSync, Verbose, TEXT("Sending StateEvent: %s, at frame: %d"),
+				*SequencerStateEvent.State.SequenceObjectPath, SequencerStateEvent.State.Time.Time.FrameNumber.Value);
 			Session->SendCustomEvent(SequencerStateEvent, Session->GetSessionServerEndpointId(), EConcertMessageFlags::ReliableOrdered);
 
 			// If we're playing then ensure we are set to controller (driving the playback on all clients)
@@ -400,7 +401,7 @@ void FConcertClientSequencerManager::ApplyTransportCloseEvent(const FConcertSequ
 	FConcertSequencerState* SequencerState = SequencerStates.Find(*PendingClose.SequenceObjectPath);
 	if (SequencerState)
 	{
-		// if the event was that a sequencer that was in master playback mode was closed, stop playback
+		// if the event was that a sequencer that was in controller playback mode was closed, stop playback
 		if (PendingClose.bControllerClose)
 		{
 			SequencerState->PlayerStatus = EConcertMovieScenePlayerStatus::Stopped;
@@ -496,7 +497,7 @@ void FConcertClientSequencerManager::DestroyPlayer(FSequencePlayer& Player, ALev
 
 void FConcertClientSequencerManager::ApplyCloseToPlayers(const FConcertSequencerCloseEvent& InEvent)
 {
-	UE_LOG(LogConcertSequencerSync, Verbose, TEXT("CloseEvent: %s, is from master: %d"), *InEvent.SequenceObjectPath, InEvent.bControllerClose);
+	UE_LOG(LogConcertSequencerSync, Verbose, TEXT("Apply CloseEvent: %s, is from controller: %d"), *InEvent.SequenceObjectPath, InEvent.bControllerClose);
 	FSequencePlayer* Player = SequencePlayers.Find(*InEvent.SequenceObjectPath);
 	ALevelSequenceActor* LevelSequenceActor = Player ? Player->Actor.Get() : nullptr;
 
@@ -537,7 +538,7 @@ void FConcertClientSequencerManager::ApplyTimeAdjustmentEvent(const FConcertSequ
 		return;
 	}
 
-	UE_LOG(LogConcertSequencerSync, Verbose, TEXT("Apply Time adjustment Event: %s"), *TimeEvent.SequenceObjectPath);
+	UE_LOG(LogConcertSequencerSync, Verbose, TEXT("Apply TimeAdjustmentEvent: %s"), *TimeEvent.SequenceObjectPath);
 	TGuardValue<bool> ReentrancyGuard(bRespondingToTransportEvent, true);
 
 	if (GIsEditor)
@@ -648,7 +649,7 @@ void FConcertClientSequencerManager::ApplyTransportEvent(const FConcertSequencer
 
 void FConcertClientSequencerManager::ApplyEventToSequencers(const FConcertSequencerState& EventState)
 {
-	UE_LOG(LogConcertSequencerSync, Verbose, TEXT("ApplyEvent: %s, at frame: %d"), *EventState.SequenceObjectPath, EventState.Time.Time.FrameNumber.Value);
+	UE_LOG(LogConcertSequencerSync, Verbose, TEXT("ApplyEventToSequencers: %s, at frame: %d"), *EventState.SequenceObjectPath, EventState.Time.Time.FrameNumber.Value);
 	FConcertSequencerState& SequencerState = SequencerStates.FindOrAdd(*EventState.SequenceObjectPath);
 
 	// Record the Sequencer State
@@ -742,34 +743,54 @@ void FConcertClientSequencerManager::ApplyEventToSequencers(const FConcertSequen
 namespace UE::Private::ConcertClientSequencerManager
 {
 
-int32 GetDuration(const FFrameNumberRange& PlayRange)
+void ApplyPlayRangeToPlayer(ULevelSequencePlayer* Player, const FFrameNumberRange& PlayRange)
 {
-	const TRangeBound<FFrameNumber>& UpperBound = PlayRange.GetUpperBound();
-	if (UpperBound.IsClosed())
+	check(Player);
+
+	const bool bLowerBoundClosed = PlayRange.GetLowerBound().IsClosed();
+	const bool bUpperBoundClosed = PlayRange.GetUpperBound().IsClosed();
+
+	if (!ensureMsgf(bLowerBoundClosed, TEXT("PlayRange lower bound is open, which is not supported.")))
 	{
-		const FFrameNumber SrcStartFrame = UE::MovieScene::DiscreteInclusiveLower(PlayRange);
-		const FFrameNumber SrcEndFrame   = UE::MovieScene::DiscreteExclusiveUpper(PlayRange);
-		return (SrcEndFrame - SrcStartFrame).Value;
+		return;
 	}
-	return TNumericLimits<int32>::Max() - 1;
-}
 
-void ApplyPlayRangeToPlayer(ULevelSequencePlayer* Player, const FFrameNumberRange& PlayRange)
-{
-	check(Player);
+	if (!Player->GetSequence() || !Player->GetSequence()->GetMovieScene())
+	{
+		UE_LOG(LogConcertSequencerSync, Warning, TEXT("ApplyPlayRangeToPlayer (%s): Missing sequence or scene"),
+			*Player->GetSequenceName());
+		return;
+	}
+
+	// Convert passed-in range from tick resolution to display rate.
+	UMovieScene* MovieScene = Player->GetSequence()->GetMovieScene();
+	const FFrameRate TickRate = MovieScene->GetTickResolution();
+	const FFrameRate PlayerRate = Player->GetFrameRate();
+
+	const FFrameNumber NewStartFrame = ConvertFrameTime(
+		UE::MovieScene::DiscreteInclusiveLower(PlayRange), TickRate, PlayerRate).FloorToFrame();
 
-	if (ensureMsgf(PlayRange.GetLowerBound().IsClosed(), TEXT("PlayRange lower bound is open, which is not supported.")))
+	int32 NewDuration;
+	if (bUpperBoundClosed)
 	{
-		const FFrameNumber SrcStartFrame = UE::MovieScene::DiscreteInclusiveLower(PlayRange);
+		const FFrameNumber NewEndFrame = ConvertFrameTime(
+			UE::MovieScene::DiscreteExclusiveUpper(PlayRange), TickRate, PlayerRate).FloorToFrame();
 
-		const int32 CurrentDuration = Player->GetFrameDuration();
-		const int32 NewDuration = GetDuration(PlayRange);
+		NewDuration = (NewEndFrame - NewStartFrame).Value;
+	}
+	else
+	{
+		NewDuration = TNumericLimits<int32>::Max() - 1;
+	}
 
-		const FQualifiedFrameTime QualifiedStartTime = Player->GetStartTime();
-		if (CurrentDuration != NewDuration || QualifiedStartTime.Time.GetFrame() != SrcStartFrame)
-		{
-			Player->SetFrameRange(SrcStartFrame.Value, NewDuration);
-		}
+	const int32 CurrentDuration = Player->GetFrameDuration();
+	const FFrameNumber CurrentStartFrame = Player->GetStartTime().Time.GetFrame();
+	if (CurrentDuration != NewDuration || CurrentStartFrame != NewStartFrame)
+	{
+		UE_LOG(LogConcertSequencerSync, Verbose, TEXT("SetFrameRange (%s): Start %d, duration %d (was %d, %d)"),
+			*Player->GetSequenceName(), NewStartFrame.Value, NewDuration, CurrentStartFrame.Value, CurrentDuration);
+
+		Player->SetFrameRange(NewStartFrame.Value, NewDuration);
 	}
 }
 
