Add a delegate for post workspace finalization.

diff --git a/Engine/Plugins/Developer/Concert/ConcertSync/ConcertSyncClient/Source/ConcertSyncClient/Private/ConcertClientWorkspace.cpp b/Engine/Plugins/Developer/Concert/ConcertSync/ConcertSyncClient/Source/ConcertSyncClient/Private/ConcertClientWorkspace.cpp
--- a/Engine/Plugins/Developer/Concert/ConcertSync/ConcertSyncClient/Source/ConcertSyncClient/Private/ConcertClientWorkspace.cpp
+++ b/Engine/Plugins/Developer/Concert/ConcertSync/ConcertSyncClient/Source/ConcertSyncClient/Private/ConcertClientWorkspace.cpp
@@ -446,6 +446,11 @@ FOnWorkspaceSynchronized& FConcertClientWorkspace::OnWorkspaceSynchronized()
 	return OnWorkspaceSyncedDelegate;
 }
 
+FOnFinalizeWorkspaceSyncCompleted& FConcertClientWorkspace::OnFinalizeWorkspaceSyncCompleted()
+{
+	return OnFinalizeWorkspaceSyncCompletedDelegate;
+}
+
 IConcertClientDataStore& FConcertClientWorkspace::GetDataStore()
 {
 	return *DataStore;
@@ -841,6 +846,7 @@ void FConcertClientWorkspace::OnEndFrame()
 		// Finalize the sync
 		bHasSyncedWorkspace = true;
 		FConcertSlowTaskStackWorkaround::Get().PopTask(MoveTemp(InitialSyncSlowTask));
+		OnFinalizeWorkspaceSyncCompletedDelegate.Broadcast();
 	}
 
 	if (bHasSyncedWorkspace && CanProcessPendingPackages() && !ConcertSyncClientUtil::UserIsEditing())
diff --git a/Engine/Plugins/Developer/Concert/ConcertSync/ConcertSyncClient/Source/ConcertSyncClient/Private/ConcertClientWorkspace.h b/Engine/Plugins/Developer/Concert/ConcertSync/ConcertSyncClient/Source/ConcertSyncClient/Private/ConcertClientWorkspace.h
--- a/Engine/Plugins/Developer/Concert/ConcertSync/ConcertSyncClient/Source/ConcertSyncClient/Private/ConcertClientWorkspace.h
+++ b/Engine/Plugins/Developer/Concert/ConcertSync/ConcertSyncClient/Source/ConcertSyncClient/Private/ConcertClientWorkspace.h
@@ -51,6 +51,8 @@ public:
 	virtual int64 GetLastActivityId() const override;
 	virtual FOnActivityAddedOrUpdated& OnActivityAddedOrUpdated() override;
 	virtual FOnWorkspaceSynchronized& OnWorkspaceSynchronized() override;
+	virtual FOnFinalizeWorkspaceSyncCompleted& OnFinalizeWorkspaceSyncCompleted() override;
+
 	virtual IConcertClientDataStore& GetDataStore() override;
 	virtual bool IsAssetModifiedByOtherClients(const FName& AssetName, int32* OutOtherClientsWithModifNum, TArray<FConcertClientInfo>* OutOtherClientsWithModifInfo, int32 OtherClientsWithModifMaxFetchNum) const override;
 	virtual void SetIgnoreOnRestoreFlagForEmittedActivities(bool bIgnore) override;
@@ -232,7 +234,10 @@ private:
 
 	/** The delegate called every time the workspace is synced. */
 	FOnWorkspaceSynchronized OnWorkspaceSyncedDelegate;
-	
+
+	/** The delegate called after the workspace has been synced and finalized. */
+	FOnFinalizeWorkspaceSyncCompleted OnFinalizeWorkspaceSyncCompletedDelegate;
+
 	/** The session key/value store proxy. The real store is held by the server and shared across all clients. */
 	TUniquePtr<FConcertClientDataStore> DataStore;
 
diff --git a/Engine/Plugins/Developer/Concert/ConcertSync/ConcertSyncClient/Source/ConcertSyncClient/Public/IConcertClientWorkspace.h b/Engine/Plugins/Developer/Concert/ConcertSync/ConcertSyncClient/Source/ConcertSyncClient/Public/IConcertClientWorkspace.h
--- a/Engine/Plugins/Developer/Concert/ConcertSync/ConcertSyncClient/Source/ConcertSyncClient/Public/IConcertClientWorkspace.h
+++ b/Engine/Plugins/Developer/Concert/ConcertSync/ConcertSyncClient/Source/ConcertSyncClient/Public/IConcertClientWorkspace.h
@@ -16,6 +16,7 @@ DECLARE_DELEGATE_RetVal(bool, FCanFinalizeWorkspaceDelegate);
 DECLARE_DELEGATE_RetVal(bool, FCanProcessPendingPackages);
 
 DECLARE_MULTICAST_DELEGATE(FOnWorkspaceSynchronized);
+DECLARE_MULTICAST_DELEGATE(FOnFinalizeWorkspaceSyncCompleted);
 DECLARE_MULTICAST_DELEGATE_ThreeParams(FOnActivityAddedOrUpdated, const FConcertClientInfo&/*InClientInfo*/, const FConcertSyncActivity&/*InActivity*/, const FStructOnScope&/*InActivitySummary*/);
 
 class IConcertClientWorkspace
@@ -148,10 +149,16 @@ public:
 	virtual bool FindPackageEvent(const int64 PackageEventId, FConcertSyncPackageEventMetaData& OutPackageEvent) const = 0;
 
 	/**
-	 * @return the delegate called every time the workspace is synced.
+	 * @return the delegate called every time the workspace is synced. This is when the server has indicated that all
+	 * activities have been sent and ready to be finalized in the current workspace.
 	 */
 	virtual FOnWorkspaceSynchronized& OnWorkspaceSynchronized() = 0;
 
+	/**
+	 * @return The delegate called after a workspace has been completely synced and finalized.  All transactions are posted and packages have been loaded.
+	 */
+	virtual FOnFinalizeWorkspaceSyncCompleted& OnFinalizeWorkspaceSyncCompleted() = 0;
+
 	/**
 	 * This delegate allows user to defer the finalization of a sync workspace. This is for situtations where multiple
 	 * client nodes need to be finalized at the same point in time.  The delegate function should return true when workspace
