Make the default sequencer settings to be true.

diff --git a/Engine/Plugins/Developer/Concert/ConcertSync/ConcertSyncClient/Source/ConcertSyncClient/Private/ConcertClientSequencerManager.cpp b/Engine/Plugins/Developer/Concert/ConcertSync/ConcertSyncClient/Source/ConcertSyncClient/Private/ConcertClientSequencerManager.cpp
--- a/Engine/Plugins/Developer/Concert/ConcertSync/ConcertSyncClient/Source/ConcertSyncClient/Private/ConcertClientSequencerManager.cpp
+++ b/Engine/Plugins/Developer/Concert/ConcertSync/ConcertSyncClient/Source/ConcertSyncClient/Private/ConcertClientSequencerManager.cpp
@@ -43,10 +43,10 @@ static TAutoConsoleVariable<int32> CVarEnablePlaybackSync(TEXT("Concert.EnableSe
 static TAutoConsoleVariable<int32> CVarEnableSequencePlayer(TEXT("Concert.EnableSequencePlayer"), 1, TEXT("Enable Concert Sequence Players on `-game` client."));
 
 // Enable opening Sequencer on remote machine whenever a sequencer is opened, if both instance have this option on.
-static TAutoConsoleVariable<int32> CVarEnableRemoteSequencerOpen(TEXT("Concert.EnableOpenRemoteSequencer"), 0, TEXT("Enable Concert remote Sequencer opening."));
+static TAutoConsoleVariable<int32> CVarEnableRemoteSequencerOpen(TEXT("Concert.EnableOpenRemoteSequencer"), 1, TEXT("Enable Concert remote Sequencer opening."));
 
 // Enable opening Sequencer on remote machine whenever a sequencer is opened, if both instance have this option on.
-static TAutoConsoleVariable<int32> CVarEnableUnrelatedTimelineSync(TEXT("Concert.EnableUnrelatedTimelineSync"), 0, TEXT("Enable syncing unrelated sequencer timeline."));
+static TAutoConsoleVariable<int32> CVarEnableUnrelatedTimelineSync(TEXT("Concert.EnableUnrelatedTimelineSync"), 1, TEXT("Enable syncing unrelated sequencer timeline."));
 
 // Enable always closing player on remote machine whenever a sequencer is closed on an editor.
 static TAutoConsoleVariable<int32> CVarAlwaysCloseGamePlayerOnCloseEvent(
