Exclude packages from external actors from the dirty list; hot reload will handle proper reload of the asset.

diff --git a/Engine/Plugins/Developer/Concert/ConcertSync/ConcertSyncClient/Source/ConcertSyncClient/Private/ConcertClientPackageManager.cpp b/Engine/Plugins/Developer/Concert/ConcertSync/ConcertSyncClient/Source/ConcertSyncClient/Private/ConcertClientPackageManager.cpp
--- a/Engine/Plugins/Developer/Concert/ConcertSync/ConcertSyncClient/Source/ConcertSyncClient/Private/ConcertClientPackageManager.cpp
+++ b/Engine/Plugins/Developer/Concert/ConcertSync/ConcertSyncClient/Source/ConcertSyncClient/Private/ConcertClientPackageManager.cpp
@@ -23,6 +23,7 @@
 #include "HAL/PlatformFileManager.h"
 
 #include "ISourceControlModule.h"
+#include "UObject/UObjectHash.h"
 
 #if WITH_EDITOR
 	#include "Editor.h"
@@ -125,9 +126,25 @@ FConcertClientPackageManager::~FConcertClientPackageManager()
 	}
 }
 
+bool PackageContainsExternalActors(UPackage* InPackage)
+{
+	bool bContainsExternalActors = false;
+	ForEachObjectWithPackage(
+		InPackage, [&bContainsExternalActors](UObject* InObject) mutable -> bool {
+			if (InObject->IsPackageExternal()) {
+				bContainsExternalActors = true;
+				// Break from our loop.
+				return false;
+			}
+			return true;
+		});
+	return bContainsExternalActors;
+}
+
 bool FConcertClientPackageManager::ShouldIgnorePackageDirtyEvent(class UPackage* InPackage) const
 {
 	return InPackage == GetTransientPackage()
+		|| PackageContainsExternalActors(InPackage)
 		|| InPackage->HasAnyFlags(RF_Transient)
 		|| InPackage->HasAnyPackageFlags(PKG_PlayInEditor | PKG_CompiledIn) // CompiledIn packages are not considered content for MU. (ex when changing some plugin settings like /Script/DisasterRecoveryClient)
 		|| bIgnorePackageDirtyEvent
@@ -405,7 +422,11 @@ void FConcertClientPackageManager::HandlePackageRejectedEvent(const FConcertSess
 void FConcertClientPackageManager::HandlePackageDirtyStateChanged(UPackage* InPackage)
 {
 	check(!InPackage->HasAnyFlags(RF_Transient) || InPackage != GetTransientPackage());
-	if (InPackage->IsDirty() && !InPackage->HasAnyPackageFlags(PKG_CompiledIn | PKG_InMemoryOnly)) // Dirty packages are tracked for purge/reload, but 'compiled in' and 'in memory' cannot be hot purged/reloaded.
+
+	// Dirty packages are tracked for purge/reload, but 'compiled in' and
+	// 'in memory' cannot be hot purged/reloaded.
+	//
+	if (InPackage->IsDirty() && !InPackage->HasAnyPackageFlags(PKG_CompiledIn | PKG_InMemoryOnly))
 	{
 		DirtyPackages.Add(InPackage->GetFName());
 	}
