ConcertSyncClient: add a workspace OnWorkspaceEndFrameCompleted delegate and chain the sequence manager's OnEndFrame to it

This changes the end frame handling of FConcertClientSequencerManager so that its
OnEndFrame() handler that was previously invoked by the core engine's OnEndFrame()
delegate is instead now invoked by a new delegate on the workspace that broadcasts
when it has finished its own OnEndFrame(). This ensures that any workspace handling
in its OnEndFrame() completes before the sequence manager's handler runs. In particular,
we want to be sure that any asynchronous package writes have completed and packages
have been reloaded before the sequence manager tries to make use of those packages.

diff --git a/Engine/Plugins/Developer/Concert/ConcertSync/ConcertSyncClient/Source/ConcertSyncClient/Private/ConcertClientSequencerManager.cpp b/Engine/Plugins/Developer/Concert/ConcertSync/ConcertSyncClient/Source/ConcertSyncClient/Private/ConcertClientSequencerManager.cpp
--- a/Engine/Plugins/Developer/Concert/ConcertSync/ConcertSyncClient/Source/ConcertSyncClient/Private/ConcertClientSequencerManager.cpp
+++ b/Engine/Plugins/Developer/Concert/ConcertSync/ConcertSyncClient/Source/ConcertSyncClient/Private/ConcertClientSequencerManager.cpp
@@ -7,7 +7,6 @@
 #include "Misc/AssertionMacros.h"
 #include "Misc/QualifiedFrameTime.h"
 #include "Modules/ModuleManager.h"
-#include "Misc/CoreDelegates.h"
 #include "Logging/LogMacros.h"
 
 #include "IConcertClient.h"
@@ -24,6 +23,7 @@
 #include "LevelSequenceActor.h"
 
 #include "UObject/PackageReload.h"
+#include "UObject/UObjectGlobals.h"
 
 #if WITH_EDITOR
 	#include "ISequencerModule.h"
@@ -61,7 +61,7 @@ FConcertClientSequencerManager::FConcertClientSequencerManager(IConcertSyncClien
 
 	ISequencerModule& SequencerModule = FModuleManager::Get().LoadModuleChecked<ISequencerModule>("Sequencer");
 	OnSequencerCreatedHandle = SequencerModule.RegisterOnSequencerCreated(FOnSequencerCreated::FDelegate::CreateRaw(this, &FConcertClientSequencerManager::OnSequencerCreated));
-	FCoreDelegates::OnEndFrame.AddRaw(this, &FConcertClientSequencerManager::OnEndFrame);
+
 	FCoreUObjectDelegates::OnPackageReloaded.AddRaw(this, &FConcertClientSequencerManager::HandleAssetReload);
 }
 
@@ -73,7 +73,7 @@ FConcertClientSequencerManager::~FConcertClientSequencerManager()
 		SequencerModulePtr->UnregisterOnSequencerCreated(OnSequencerCreatedHandle);
 	}
 
-	FCoreDelegates::OnEndFrame.RemoveAll(this);
+	SetActiveWorkspace(nullptr);
 
 	for (FOpenSequencerData& OpenSequencer : OpenSequencers)
 	{
@@ -154,7 +154,17 @@ FConcertClientSequencerManager::GatherRootSequencersByState(const FString& InSeq
 
 void FConcertClientSequencerManager::SetActiveWorkspace(TSharedPtr<FConcertClientWorkspace> InWorkspace)
 {
+	if (TSharedPtr<FConcertClientWorkspace> SharedWorkspace = Workspace.Pin())
+	{
+		SharedWorkspace->OnWorkspaceEndFrameCompleted().RemoveAll(this);
+	}
+
 	Workspace = InWorkspace;
+
+	if (InWorkspace.IsValid())
+	{
+		InWorkspace->OnWorkspaceEndFrameCompleted().AddRaw(this, &FConcertClientSequencerManager::OnWorkspaceEndFrameCompleted);
+	}
 }
 
 float FConcertClientSequencerManager::GetLatencyCompensationMs() const
@@ -925,20 +935,8 @@ bool FConcertClientSequencerManager::CanSendSequencerEvent(const FString& Object
 	return true;
 }
 
-void FConcertClientSequencerManager::OnEndFrame()
+void FConcertClientSequencerManager::OnWorkspaceEndFrameCompleted()
 {
-	TSharedPtr<FConcertClientWorkspace> SharedWorkspace = Workspace.Pin();
-	if (SharedWorkspace)
-	{
-		if (!SharedWorkspace->CanProcessPendingPackages())
-		{
-			// There is currently a lock on the workspace.  We should wait for those to finish before processing
-			// sequencer events.
-			//
-			return;
-		}
-	}
-
 	for (const FConcertSequencerCloseEvent& CloseEvent: PendingSequenceCloseEvents)
 	{
 		ApplyTransportCloseEvent(CloseEvent);
diff --git a/Engine/Plugins/Developer/Concert/ConcertSync/ConcertSyncClient/Source/ConcertSyncClient/Private/ConcertClientSequencerManager.h b/Engine/Plugins/Developer/Concert/ConcertSync/ConcertSyncClient/Source/ConcertSyncClient/Private/ConcertClientSequencerManager.h
--- a/Engine/Plugins/Developer/Concert/ConcertSync/ConcertSyncClient/Source/ConcertSyncClient/Private/ConcertClientSequencerManager.h
+++ b/Engine/Plugins/Developer/Concert/ConcertSync/ConcertSyncClient/Source/ConcertSyncClient/Private/ConcertClientSequencerManager.h
@@ -205,9 +205,13 @@ private:
 	void OnSequencerCreated(TSharedRef<ISequencer> InSequencer);
 
 	/**
-	 * Handle the end of frame callback to apply pending sequencer events
+	 * Handle the end of frame callback to apply pending sequencer events.
+	 * 
+	 * Note that this is called through a delegate on the workspace instead of
+	 * being called directly by a core engine delegate. This is to ensure that
+	 * all workspace operations have completed first.
 	 */
-	void OnEndFrame();
+	void OnWorkspaceEndFrameCompleted();
 
 	/**
 	 * Apply a Sequencer open event
diff --git a/Engine/Plugins/Developer/Concert/ConcertSync/ConcertSyncClient/Source/ConcertSyncClient/Private/ConcertClientWorkspace.cpp b/Engine/Plugins/Developer/Concert/ConcertSync/ConcertSyncClient/Source/ConcertSyncClient/Private/ConcertClientWorkspace.cpp
--- a/Engine/Plugins/Developer/Concert/ConcertSync/ConcertSyncClient/Source/ConcertSyncClient/Private/ConcertClientWorkspace.cpp
+++ b/Engine/Plugins/Developer/Concert/ConcertSync/ConcertSyncClient/Source/ConcertSyncClient/Private/ConcertClientWorkspace.cpp
@@ -452,6 +452,11 @@ FOnFinalizeWorkspaceSyncCompleted& FConcertClientWorkspace::OnFinalizeWorkspaceS
 	return OnFinalizeWorkspaceSyncCompletedDelegate;
 }
 
+FOnWorkspaceEndFrameCompleted& FConcertClientWorkspace::OnWorkspaceEndFrameCompleted()
+{
+	return OnWorkspaceEndFrameCompletedDelegate;
+}
+
 IConcertClientDataStore& FConcertClientWorkspace::GetDataStore()
 {
 	return *DataStore;
@@ -871,6 +876,8 @@ void FConcertClientWorkspace::OnEndFrame()
 			LiveSession->GetSession().SendCustomEvent(StateChangeEvent, LiveSession->GetSession().GetSessionServerEndpointId(), EConcertMessageFlags::ReliableOrdered);
 			bPendingStopIgnoringActivityOnRestore = false;
 		}
+
+		OnWorkspaceEndFrameCompletedDelegate.Broadcast();
 	}
 	LiveSession->GetSessionDatabase().UpdateAsynchronousTasks();
 	IConcertClientRef ConcertClient = OwnerSyncClient->GetConcertClient();
diff --git a/Engine/Plugins/Developer/Concert/ConcertSync/ConcertSyncClient/Source/ConcertSyncClient/Private/ConcertClientWorkspace.h b/Engine/Plugins/Developer/Concert/ConcertSync/ConcertSyncClient/Source/ConcertSyncClient/Private/ConcertClientWorkspace.h
--- a/Engine/Plugins/Developer/Concert/ConcertSync/ConcertSyncClient/Source/ConcertSyncClient/Private/ConcertClientWorkspace.h
+++ b/Engine/Plugins/Developer/Concert/ConcertSync/ConcertSyncClient/Source/ConcertSyncClient/Private/ConcertClientWorkspace.h
@@ -10,6 +10,7 @@
 #include "ConcertSyncSessionFlags.h"
 #include "ConcertWorkspaceMessages.h"
 #include "ConcertClientWorkspaceData.h"
+#include "Delegates/Delegate.h"
 
 class IConcertClientSession;
 class IConcertClientPackageBridge;
@@ -25,6 +26,8 @@ class FConcertClientDataStore;
 
 struct FScopedSlowTask;
 
+DECLARE_MULTICAST_DELEGATE(FOnWorkspaceEndFrameCompleted);
+
 class FConcertClientWorkspace : public IConcertClientWorkspace
 {
 public:
@@ -71,6 +74,9 @@ public:
 	/** Indicates if we can finalize activity sync. */
 	bool CanFinalize() const;
 
+	/** Gets the delegate called after the workspace has completed its handling at the end of an engine frame. */
+	FOnWorkspaceEndFrameCompleted& OnWorkspaceEndFrameCompleted();
+
 private:
 
 	/** Bind the workspace to this session. */
@@ -238,6 +244,9 @@ private:
 	/** The delegate called after the workspace has been synced and finalized. */
 	FOnFinalizeWorkspaceSyncCompleted OnFinalizeWorkspaceSyncCompletedDelegate;
 
+	/** The delegate called after the workspace has completed its handling at the end of an engine frame. */
+	FOnWorkspaceEndFrameCompleted OnWorkspaceEndFrameCompletedDelegate;
+
 	/** The session key/value store proxy. The real store is held by the server and shared across all clients. */
 	TUniquePtr<FConcertClientDataStore> DataStore;
 
