ConcertSyncClient: flesh out logging in sequencer manager and other misc. clean-up

These changes should not introduce any change in functionality, but should
instead make it a bit easier to tell when events are being received and held to
be handled at the end of frame versus when they are actually handled. This
also includes a few bits of other clean-up.

diff --git a/Engine/Plugins/Developer/Concert/ConcertSync/ConcertSyncClient/Source/ConcertSyncClient/Private/ConcertClientSequencerManager.cpp b/Engine/Plugins/Developer/Concert/ConcertSync/ConcertSyncClient/Source/ConcertSyncClient/Private/ConcertClientSequencerManager.cpp
--- a/Engine/Plugins/Developer/Concert/ConcertSync/ConcertSyncClient/Source/ConcertSyncClient/Private/ConcertClientSequencerManager.cpp
+++ b/Engine/Plugins/Developer/Concert/ConcertSync/ConcertSyncClient/Source/ConcertSyncClient/Private/ConcertClientSequencerManager.cpp
@@ -93,14 +93,18 @@ void FConcertClientSequencerManager::OnSequencerCreated(TSharedRef<ISequencer> I
 	UMovieSceneSequence* Sequence = InSequencer->GetRootMovieSceneSequence();
 	check(Sequence != nullptr);
 
-	if (!SequencerStates.Contains(*Sequence->GetPathName()))
+	const FString SequenceObjectPath = Sequence->GetPathName();
+
+	UE_LOG(LogConcertSequencerSync, Verbose, TEXT("OnSequencerCreated: %s"), *SequenceObjectPath);
+
+	if (!SequencerStates.Contains(*SequenceObjectPath))
 	{
 		FConcertSequencerState NewState;
 		NewState.Time = InSequencer->GetGlobalTime();
-		SequencerStates.Add(*Sequence->GetPathName(), NewState);
+		SequencerStates.Add(*SequenceObjectPath, NewState);
 	}
 
-	FConcertSequencerState& SequencerState = SequencerStates.FindOrAdd(*Sequence->GetPathName());
+	FConcertSequencerState& SequencerState = SequencerStates.FindOrAdd(*SequenceObjectPath);
 
 	// Setup the Sequencer
 	FOpenSequencerData OpenSequencer;
@@ -123,12 +127,12 @@ void FConcertClientSequencerManager::OnSequencerCreated(TSharedRef<ISequencer> I
 	{
 		if (TSharedPtr<IConcertClientSession> Session = WeakSession.Pin())
 		{
-			if (CanSendSequencerEvent(Sequence->GetPathName()))
+			if (CanSendSequencerEvent(SequenceObjectPath))
 			{
 				FConcertSequencerOpenEvent OpenEvent;
-				OpenEvent.SequenceObjectPath = Sequence->GetPathName();
+				OpenEvent.SequenceObjectPath = SequenceObjectPath;
 
-				UE_LOG(LogConcertSequencerSync, Verbose, TEXT("Sending OpenEvent: %s"), *OpenEvent.SequenceObjectPath);
+				UE_LOG(LogConcertSequencerSync, Verbose, TEXT("    Sending OpenEvent: %s"), *OpenEvent.SequenceObjectPath);
 				Session->SendCustomEvent(OpenEvent, Session->GetSessionServerEndpointId(), EConcertMessageFlags::ReliableOrdered);
 			}
 		}
@@ -184,7 +188,7 @@ void FConcertClientSequencerManager::Register(TSharedRef<IConcertClientSession>
 	InSession->RegisterCustomEventHandler<FConcertSequencerStateEvent>(this, &FConcertClientSequencerManager::OnTransportEvent);
 	InSession->RegisterCustomEventHandler<FConcertSequencerCloseEvent>(this, &FConcertClientSequencerManager::OnCloseEvent);
 	InSession->RegisterCustomEventHandler<FConcertSequencerOpenEvent>(this, &FConcertClientSequencerManager::OnOpenEvent);
-	InSession->RegisterCustomEventHandler<FConcertSequencerStateSyncEvent>(this, &FConcertClientSequencerManager::OnSyncEvent);
+	InSession->RegisterCustomEventHandler<FConcertSequencerStateSyncEvent>(this, &FConcertClientSequencerManager::OnStateSyncEvent);
 	InSession->RegisterCustomEventHandler<FConcertSequencerTimeAdjustmentEvent>(this, &FConcertClientSequencerManager::OnTimeAdjustmentEvent);
 }
 
@@ -240,6 +244,11 @@ bool FConcertClientSequencerManager::ShouldAlwaysCloseGameSequencerPlayer() cons
 
 void FConcertClientSequencerManager::OnSequencerClosed(TSharedRef<ISequencer> InSequencer)
 {
+	const UMovieSceneSequence* Sequence = InSequencer->GetRootMovieSceneSequence();
+	const FString SequenceObjectPath = Sequence->GetPathName();
+
+	UE_LOG(LogConcertSequencerSync, Verbose, TEXT("OnSequencerClosed: %s"), *SequenceObjectPath);
+
 	// Find the associated open sequencer index
 	int Index = 0;
 	for (; Index < OpenSequencers.Num(); ++Index)
@@ -257,31 +266,26 @@ void FConcertClientSequencerManager::OnSequencerClosed(TSharedRef<ISequencer> In
 
 	FOpenSequencerData& ClosingSequencer = OpenSequencers[Index];
 
-	// Send close event to server and put back playback mode to undefined
-	TSharedPtr<IConcertClientSession> Session = WeakSession.Pin();
-	if (Session.IsValid())
+	if (Sequence)
 	{
-		// Find the associated sequence path name
-		UMovieSceneSequence* Sequence = InSequencer->GetRootMovieSceneSequence();
-		if (Sequence)
+		// Send close event to server and put back playback mode to undefined
+		if (TSharedPtr<IConcertClientSession> Session = WeakSession.Pin())
 		{
-			if (CanSendSequencerEvent(Sequence->GetPathName()))
+			// Find the associated sequence path name
+			if (CanSendSequencerEvent(SequenceObjectPath))
 			{
-				UE_LOG(LogConcertSequencerSync, Verbose, TEXT("Sending CloseEvent: %s"), *Sequence->GetPathName());
+				UE_LOG(LogConcertSequencerSync, Verbose, TEXT("    Sending CloseEvent: %s"), *SequenceObjectPath);
 
 				FConcertSequencerCloseEvent CloseEvent;
 				CloseEvent.bControllerClose = ClosingSequencer.PlaybackMode == EPlaybackMode::Controller; // this sequencer had control over the sequence playback
-				CloseEvent.SequenceObjectPath = *Sequence->GetPathName();
+				CloseEvent.SequenceObjectPath = SequenceObjectPath;
 				Session->SendCustomEvent(CloseEvent, Session->GetSessionServerEndpointId(), EConcertMessageFlags::ReliableOrdered);
 			}
 		}
-	}
-	else
-	{
-		UMovieSceneSequence* Sequence = InSequencer->GetRootMovieSceneSequence();
-		if (Sequence)
+		else
 		{
-			SequencerStates.Remove(*Sequence->GetPathName());
+			UE_LOG(LogConcertSequencerSync, VeryVerbose, TEXT("    Removing Sequencer state for sequence: %s"), *SequenceObjectPath);
+			SequencerStates.Remove(*SequenceObjectPath);
 		}
 	}
 
@@ -289,8 +293,10 @@ void FConcertClientSequencerManager::OnSequencerClosed(TSharedRef<ISequencer> In
 	OpenSequencers.RemoveAtSwap(Index);
 }
 
-void FConcertClientSequencerManager::OnSyncEvent(const FConcertSessionContext& InEventContext, const FConcertSequencerStateSyncEvent& InEvent)
+void FConcertClientSequencerManager::OnStateSyncEvent(const FConcertSessionContext& InEventContext, const FConcertSequencerStateSyncEvent& InEvent)
 {
+	UE_LOG(LogConcertSequencerSync, Verbose, TEXT("Event Received - StateSync: Syncing %d Sequencer states"), InEvent.SequencerStates.Num());
+
 	for (const auto& State : InEvent.SequencerStates)
 	{
 		FConcertSequencerState& SequencerState = SequencerStates.FindOrAdd(*State.SequenceObjectPath);
@@ -303,7 +309,7 @@ void FConcertClientSequencerManager::OnSyncEvent(const FConcertSessionContext& I
 		// Sequencers.
 		bool bFoundSequencerForObject = false;
 
-		UE_LOG(LogConcertSequencerSync, VeryVerbose, TEXT("OnSyncEvent: Syncing %d Sequencers for sequence %s"), OpenSequencersForObject.Num(), *SequencerState.SequenceObjectPath);
+		UE_LOG(LogConcertSequencerSync, VeryVerbose, TEXT("    Syncing %d Sequencers for sequence %s"), OpenSequencersForObject.Num(), *SequencerState.SequenceObjectPath);
 		for (FOpenSequencerData* OpenSequencer : OpenSequencersForObject)
 		{
 			TSharedPtr<ISequencer> Sequencer = OpenSequencer->WeakSequencer.Pin();
@@ -323,7 +329,7 @@ void FConcertClientSequencerManager::OnSyncEvent(const FConcertSessionContext& I
 
 		if (!bFoundSequencerForObject)
 		{
-			UE_LOG(LogConcertSequencerSync, VeryVerbose, TEXT("OnSyncEvent: No existing Sequencer with sequence %s open. Will open and sync a new one at end of frame."), *SequencerState.SequenceObjectPath);
+			UE_LOG(LogConcertSequencerSync, VeryVerbose, TEXT("    No existing Sequencer with sequence %s open. Will open and sync a new one at end of frame."), *SequencerState.SequenceObjectPath);
 			PendingSequenceOpenEvents.Add(SequencerState.SequenceObjectPath);
 			PendingSequencerEvents.Add(SequencerState);
 		}
@@ -390,13 +396,13 @@ void FConcertClientSequencerManager::OnSequencerTimeChanged(TWeakPtr<ISequencer>
 
 void FConcertClientSequencerManager::OnCloseEvent(const FConcertSessionContext&, const FConcertSequencerCloseEvent& InEvent)
 {
-	UE_LOG(LogConcertSequencerSync, Verbose, TEXT("OnCloseEvent: %s"), *InEvent.SequenceObjectPath);
+	UE_LOG(LogConcertSequencerSync, Verbose, TEXT("Event Received - Close: %s. Deferring until end of frame."), *InEvent.SequenceObjectPath);
 	PendingSequenceCloseEvents.Add(InEvent);
 }
 
 void FConcertClientSequencerManager::OnOpenEvent(const FConcertSessionContext&, const FConcertSequencerOpenEvent& InEvent)
 {
-	UE_LOG(LogConcertSequencerSync, Verbose, TEXT("OnOpenEvent: %s"), *InEvent.SequenceObjectPath);
+	UE_LOG(LogConcertSequencerSync, Verbose, TEXT("Event Received - Open: %s. Deferring until end of frame."), *InEvent.SequenceObjectPath);
 	PendingSequenceOpenEvents.Add(InEvent.SequenceObjectPath);
 }
 
@@ -428,13 +434,47 @@ void FConcertClientSequencerManager::HandleAssetReload(const EPackageReloadPhase
 	}
 }
 
-void FConcertClientSequencerManager::ApplyTransportCloseEvent(const FConcertSequencerCloseEvent& PendingClose)
+void FConcertClientSequencerManager::ApplyCloseEventToPlayers(const FConcertSequencerCloseEvent& CloseEvent)
+{
+	UE_LOG(LogConcertSequencerSync, Verbose, TEXT("    Closing sequence players: %s, is from controller: %d"), *CloseEvent.SequenceObjectPath, CloseEvent.bControllerClose);
+
+	FSequencePlayer* Player = SequencePlayers.Find(*CloseEvent.SequenceObjectPath);
+	if (!Player)
+	{
+		UE_LOG(LogConcertSequencerSync, Verbose, TEXT("        No open players to close for sequence %s"), *CloseEvent.SequenceObjectPath);
+		return;
+	}
+
+	ALevelSequenceActor* LevelSequenceActor = Player->Actor.Get();
+
+	if (CloseEvent.bControllerClose && LevelSequenceActor && LevelSequenceActor->SequencePlayer)
+	{
+		LevelSequenceActor->SequencePlayer->Stop();
+	}
+
+	if (!CanClose(CloseEvent))
+	{
+		return;
+	}
+
+	DestroyPlayer(*Player, LevelSequenceActor);
+
+	Player->Actor = nullptr;
+
+	// Always remove the player on close event. This will allow it to be re-opened on an OpenEvent.
+	UE_LOG(LogConcertSequencerSync, VeryVerbose, TEXT("    Removing sequence player for sequence: %s"), *CloseEvent.SequenceObjectPath);
+	SequencePlayers.Remove(*CloseEvent.SequenceObjectPath);
+}
+
+void FConcertClientSequencerManager::ApplyCloseEvent(const FConcertSequencerCloseEvent& CloseEvent)
 {
-	FConcertSequencerState* SequencerState = SequencerStates.Find(*PendingClose.SequenceObjectPath);
+	UE_LOG(LogConcertSequencerSync, Verbose, TEXT("Handling Event - Close: %s"), *CloseEvent.SequenceObjectPath);
+
+	FConcertSequencerState* SequencerState = SequencerStates.Find(*CloseEvent.SequenceObjectPath);
 	if (SequencerState)
 	{
 		// if the event was that a sequencer that was in controller playback mode was closed, stop playback
-		if (PendingClose.bControllerClose)
+		if (CloseEvent.bControllerClose)
 		{
 			SequencerState->PlayerStatus = EConcertMovieScenePlayerStatus::Stopped;
 			for (FOpenSequencerData* OpenSequencer : GatherRootSequencersByState(SequencerState->SequenceObjectPath))
@@ -446,18 +486,22 @@ void FConcertClientSequencerManager::ApplyTransportCloseEvent(const FConcertSequ
 		}
 		// Discard the state if it's no longer opened by anyone.
 		//
-		if (PendingClose.EditorsWithSequencerOpened == 0)
+		if (CloseEvent.EditorsWithSequencerOpened == 0)
 		{
-			SequencerStates.Remove(*PendingClose.SequenceObjectPath);
+			UE_LOG(LogConcertSequencerSync, VeryVerbose, TEXT("    Removing Sequencer state for sequence: %s"), *CloseEvent.SequenceObjectPath);
+			SequencerStates.Remove(*CloseEvent.SequenceObjectPath);
 		}
 	}
 
-	ApplyCloseToPlayers(PendingClose);
+	ApplyCloseEventToPlayers(CloseEvent);
 }
 
-void FConcertClientSequencerManager::ApplyTransportOpenEvent(const FString& SequenceObjectPath)
+void FConcertClientSequencerManager::ApplyOpenEvent(const FString& SequenceObjectPath)
 {
 	TGuardValue<bool> ReentrancyGuard(bRespondingToTransportEvent, true);
+
+	UE_LOG(LogConcertSequencerSync, Verbose, TEXT("Handling Event - Open: %s"), *SequenceObjectPath);
+
 #if WITH_EDITOR
 	if (IsSequencerRemoteOpenEnabled() && GIsEditor)
 	{
@@ -473,9 +517,11 @@ void FConcertClientSequencerManager::ApplyTransportOpenEvent(const FString& Sequ
 
 void FConcertClientSequencerManager::CreateNewSequencePlayerIfNotExists(const FString& SequenceObjectPath)
 {
-	// we do not have a player for this state yet
+	UE_LOG(LogConcertSequencerSync, Verbose, TEXT("CreateNewSequencePlayerIfNotExists: %s"), *SequenceObjectPath);
+
 	if (SequencePlayers.Contains(*SequenceObjectPath))
 	{
+		UE_LOG(LogConcertSequencerSync, Verbose, TEXT("    Sequence player already exists for sequence %s"), *SequenceObjectPath);
 		return;
 	}
 
@@ -486,13 +532,11 @@ void FConcertClientSequencerManager::CreateNewSequencePlayerIfNotExists(const FS
 	}
 	check(CurrentWorld);
 
-	ALevelSequenceActor* LevelSequenceActor = nullptr;
-
 	// Get the actual sequence
-	ULevelSequence* Sequence = LoadObject<ULevelSequence>( nullptr, *SequenceObjectPath);
-
+	ULevelSequence* Sequence = LoadObject<ULevelSequence>(nullptr, *SequenceObjectPath);
 	if (!Sequence)
 	{
+		UE_LOG(LogConcertSequencerSync, Warning, TEXT("    Failed to load level sequence while creating new sequence player: %s"), *SequenceObjectPath);
 		return;
 	}
 
@@ -501,7 +545,8 @@ void FConcertClientSequencerManager::CreateNewSequencePlayerIfNotExists(const FS
 	// Sequencer pauses at the last frame and Player Stops and goes to the first frame unless we set this flag.
 	PlaybackSettings.bPauseAtEnd = true;
 
-	// This call will initialeze LevelSequenceActor as an output parameter.
+	// This call will initialize LevelSequenceActor as an output parameter.
+	ALevelSequenceActor* LevelSequenceActor = nullptr;
 	ULevelSequencePlayer* Player = ULevelSequencePlayer::CreateLevelSequencePlayer(CurrentWorld->PersistentLevel, Sequence, PlaybackSettings, LevelSequenceActor);
 	check(Player);
 
@@ -519,6 +564,7 @@ void FConcertClientSequencerManager::DestroyPlayer(FSequencePlayer& Player, ALev
 {
 	if (LevelSequenceActor && LevelSequenceActor->SequencePlayer)
 	{
+		UE_LOG(LogConcertSequencerSync, VeryVerbose, TEXT("Destroying LevelSequenceActor: %s"), *LevelSequenceActor->GetPathName());
 		LevelSequenceActor->SequencePlayer->Stop();
 
 		LevelSequenceActor->SetSequence(nullptr);
@@ -527,59 +573,35 @@ void FConcertClientSequencerManager::DestroyPlayer(FSequencePlayer& Player, ALev
 	}
 }
 
-void FConcertClientSequencerManager::ApplyCloseToPlayers(const FConcertSequencerCloseEvent& InEvent)
-{
-	UE_LOG(LogConcertSequencerSync, Verbose, TEXT("Apply CloseEvent: %s, is from controller: %d"), *InEvent.SequenceObjectPath, InEvent.bControllerClose);
-	FSequencePlayer* Player = SequencePlayers.Find(*InEvent.SequenceObjectPath);
-	ALevelSequenceActor* LevelSequenceActor = Player ? Player->Actor.Get() : nullptr;
-
-	if (InEvent.bControllerClose && LevelSequenceActor && LevelSequenceActor->SequencePlayer)
-	{
-		LevelSequenceActor->SequencePlayer->Stop();
-	}
-
-	if (!Player || !CanClose(InEvent))
-	{
-		return;
-	}
-
-	DestroyPlayer(*Player, LevelSequenceActor);
-
-	Player->Actor = nullptr;
-
-	UE_LOG(LogConcertSequencerSync, Verbose, TEXT("Removing player: %s"), *InEvent.SequenceObjectPath);
-	// Always remove the player on close event. This will allow it to be re-opened on an OpenEvent.
-	//
-	SequencePlayers.Remove(*InEvent.SequenceObjectPath);
-}
-
 void FConcertClientSequencerManager::OnTransportEvent(const FConcertSessionContext&, const FConcertSequencerStateEvent& InEvent)
 {
 	PendingSequencerEvents.Add(InEvent.State);
 }
 
-void FConcertClientSequencerManager::OnTimeAdjustmentEvent(const FConcertSessionContext&, const FConcertSequencerTimeAdjustmentEvent& TimeAdjustmentEvent)
+void FConcertClientSequencerManager::OnTimeAdjustmentEvent(const FConcertSessionContext&, const FConcertSequencerTimeAdjustmentEvent& InEvent)
 {
-	PendingTimeAdjustmentEvents.Add(TimeAdjustmentEvent);
+	UE_LOG(LogConcertSequencerSync, Verbose, TEXT("Event Received - TimeAdjustment: %s. Deferring until end of frame."), *InEvent.SequenceObjectPath);
+	PendingTimeAdjustmentEvents.Add(InEvent);
 }
 
-void FConcertClientSequencerManager::ApplyTimeAdjustmentEvent(const FConcertSequencerTimeAdjustmentEvent& TimeEvent)
+void FConcertClientSequencerManager::ApplyTimeAdjustmentEvent(const FConcertSequencerTimeAdjustmentEvent& TimeAdjustmentEvent)
 {
 	if (bRespondingToTransportEvent)
 	{
 		return;
 	}
 
-	UE_LOG(LogConcertSequencerSync, Verbose, TEXT("Apply TimeAdjustmentEvent: %s"), *TimeEvent.SequenceObjectPath);
 	TGuardValue<bool> ReentrancyGuard(bRespondingToTransportEvent, true);
 
+	UE_LOG(LogConcertSequencerSync, Verbose, TEXT("Handling Event - TimeAdjustment: %s"), *TimeAdjustmentEvent.SequenceObjectPath);
+
 	if (GIsEditor)
 	{
-		ApplyTimeAdjustmentSequencers(TimeEvent);
+		ApplyTimeAdjustmentToSequencers(TimeAdjustmentEvent);
 	}
 	else if (CVarEnableSequencePlayer.GetValueOnAnyThread() > 0)
 	{
-		ApplyTimeAdjustmentPlayers(TimeEvent);
+		ApplyTimeAdjustmentToPlayers(TimeAdjustmentEvent);
 	}
 }
 
@@ -603,9 +625,11 @@ bool ApplyStartFrameToMovieScene(FFrameNumber StartFrame, UMovieScene* MovieScen
 	return true;
 }
 
-void FConcertClientSequencerManager::ApplyTimeAdjustmentSequencers(const FConcertSequencerTimeAdjustmentEvent& TimeEvent)
+void FConcertClientSequencerManager::ApplyTimeAdjustmentToSequencers(const FConcertSequencerTimeAdjustmentEvent& TimeAdjustmentEvent)
 {
-	for (FOpenSequencerData* OpenSequencer : GatherRootSequencersByState(TimeEvent.SequenceObjectPath))
+	UE_LOG(LogConcertSequencerSync, Verbose, TEXT("    TimeAdjustment: Updating Sequencers for sequence %s"), *TimeAdjustmentEvent.SequenceObjectPath);
+
+	for (FOpenSequencerData* OpenSequencer : GatherRootSequencersByState(TimeAdjustmentEvent.SequenceObjectPath))
 	{
 		ISequencer* Sequencer = OpenSequencer->WeakSequencer.Pin().Get();
 		// If the entry is driving playback (PlaybackMode == Controller) then we never respond time adjustment events.
@@ -619,10 +643,10 @@ void FConcertClientSequencerManager::ApplyTimeAdjustmentSequencers(const FConcer
 		UMovieScene* MovieScene = Sequencer->GetRootMovieSceneSequence()->GetMovieScene();
 		if (MovieScene)
 		{
-			if (ApplyStartFrameToMovieScene(TimeEvent.PlaybackStartFrame, MovieScene))
+			if (ApplyStartFrameToMovieScene(TimeAdjustmentEvent.PlaybackStartFrame, MovieScene))
 			{
 				FFrameRate FrameRate = MovieScene->GetTickResolution();
-				double CurrentTimeSeconds = FrameRate.AsSeconds(FFrameTime(TimeEvent.PlaybackStartFrame));
+				double CurrentTimeSeconds = FrameRate.AsSeconds(FFrameTime(TimeAdjustmentEvent.PlaybackStartFrame));
 				TRange<double> NewRange(CurrentTimeSeconds, CurrentTimeSeconds+10.0);
 				Sequencer->SetViewRange(NewRange, EViewRangeInterpolation::Immediate);
 				Sequencer->SetClampRange(Sequencer->GetViewRange());
@@ -632,12 +656,14 @@ void FConcertClientSequencerManager::ApplyTimeAdjustmentSequencers(const FConcer
 	}
 }
 
-void FConcertClientSequencerManager::ApplyTimeAdjustmentPlayers(const FConcertSequencerTimeAdjustmentEvent& TimeEvent)
+void FConcertClientSequencerManager::ApplyTimeAdjustmentToPlayers(const FConcertSequencerTimeAdjustmentEvent& TimeAdjustmentEvent)
 {
-	FSequencePlayer* SeqPlayer = SequencePlayers.Find(*TimeEvent.SequenceObjectPath);
+	UE_LOG(LogConcertSequencerSync, Verbose, TEXT("    TimeAdjustment: Updating sequence players for sequence %s"), *TimeAdjustmentEvent.SequenceObjectPath);
 
+	FSequencePlayer* SeqPlayer = SequencePlayers.Find(*TimeAdjustmentEvent.SequenceObjectPath);
 	if (!SeqPlayer)
 	{
+		UE_LOG(LogConcertSequencerSync, Verbose, TEXT("        No sequence player for sequence %s"), *TimeAdjustmentEvent.SequenceObjectPath);
 		return;
 	}
 
@@ -647,7 +673,7 @@ void FConcertClientSequencerManager::ApplyTimeAdjustmentPlayers(const FConcertSe
 		UMovieScene* MovieScene = LevelSequenceActor->LevelSequenceAsset->GetMovieScene();
 		if (MovieScene)
 		{
-			ApplyStartFrameToMovieScene(TimeEvent.PlaybackStartFrame, MovieScene);
+			ApplyStartFrameToMovieScene(TimeAdjustmentEvent.PlaybackStartFrame, MovieScene);
 		}
 	}
 }
@@ -660,9 +686,10 @@ void FConcertClientSequencerManager::ApplyTransportEvent(const FConcertSequencer
 		return;
 	}
 
-	UE_LOG(LogConcertSequencerSync, Verbose, TEXT("Apply Transport Event: %s"), *EventState.SequenceObjectPath);
 	TGuardValue<bool> ReentrancyGuard(bRespondingToTransportEvent, true);
 
+	UE_LOG(LogConcertSequencerSync, Verbose, TEXT("Handling Event - Transport: %s"), *EventState.SequenceObjectPath);
+
 	// Update the sequencer pointing to the same sequence
 	FConcertSequencerState& SequencerState = SequencerStates.FindOrAdd(*EventState.SequenceObjectPath);
 
@@ -681,7 +708,7 @@ void FConcertClientSequencerManager::ApplyTransportEvent(const FConcertSequencer
 
 void FConcertClientSequencerManager::ApplyEventToSequencers(const FConcertSequencerState& EventState)
 {
-	UE_LOG(LogConcertSequencerSync, Verbose, TEXT("ApplyEventToSequencers: %s, at frame: %d"), *EventState.SequenceObjectPath, EventState.Time.Time.FrameNumber.Value);
+	UE_LOG(LogConcertSequencerSync, Verbose, TEXT("    Transport: Update Sequencer for sequence %s, at frame: %d"), *EventState.SequenceObjectPath, EventState.Time.Time.FrameNumber.Value);
 	FConcertSequencerState& SequencerState = SequencerStates.FindOrAdd(*EventState.SequenceObjectPath);
 
 	// Record the Sequencer State
@@ -831,10 +858,12 @@ void ApplyPlayRangeToPlayer(ULevelSequencePlayer* Player, const FFrameNumberRang
 
 void FConcertClientSequencerManager::ApplyEventToPlayers(const FConcertSequencerState& EventState)
 {
-	FSequencePlayer* SeqPlayer = SequencePlayers.Find(*EventState.SequenceObjectPath);
+	UE_LOG(LogConcertSequencerSync, Verbose, TEXT("    Transport: Update sequence player for sequence %s, at frame: %d"), *EventState.SequenceObjectPath, EventState.Time.Time.FrameNumber.Value);
 
+	FSequencePlayer* SeqPlayer = SequencePlayers.Find(*EventState.SequenceObjectPath);
 	if (!SeqPlayer)
 	{
+		UE_LOG(LogConcertSequencerSync, Verbose, TEXT("        No sequence player for sequence %s"), *EventState.SequenceObjectPath);
 		return;
 	}
 
@@ -939,30 +968,33 @@ void FConcertClientSequencerManager::OnWorkspaceEndFrameCompleted()
 {
 	for (const FConcertSequencerCloseEvent& CloseEvent: PendingSequenceCloseEvents)
 	{
-		ApplyTransportCloseEvent(CloseEvent);
+		ApplyCloseEvent(CloseEvent);
 	}
 	PendingSequenceCloseEvents.Reset();
 
-	for (const TTuple<FName,FString>& Player: PendingDestroy)
+	if (!PendingDestroy.IsEmpty())
 	{
-		if (FSequencePlayer* SequencePlayer = SequencePlayers.Find(Player.Get<0>()))
+		for (const TTuple<FName, FString>& Player : PendingDestroy)
 		{
-			DestroyPlayer(*SequencePlayer, SequencePlayer->Actor.Get());
-			SequencePlayers.Remove(Player.Get<0>());
-			PendingCreate.Add(Player.Get<1>());
+			if (FSequencePlayer* SequencePlayer = SequencePlayers.Find(Player.Get<0>()))
+			{
+				DestroyPlayer(*SequencePlayer, SequencePlayer->Actor.Get());
+				SequencePlayers.Remove(Player.Get<0>());
+				PendingCreate.Add(Player.Get<1>());
+			}
 		}
-	}
 
-	if (PendingDestroy.Num())
-	{
 		if (GEngine)
 		{
 			GEngine->ForceGarbageCollection(true);
 		}
+
+		// Wait for the next frame to continue. Garbage collection will not happen
+		// immediately so we need to wait for the next frame before we re-create.
+		UE_LOG(LogConcertSequencerSync, VeryVerbose, TEXT("Running garbage collection after destroying %d sequence players. Deferring remaining pending events until next frame"), PendingDestroy.Num());
+
 		PendingDestroy.Reset();
-		// Wait for the next frame to continue.  Garbage collection will not happen immediately so we need to wait for the
-		// next frame before we re-create.
-		//
+
 		return;
 	}
 
@@ -974,7 +1006,7 @@ void FConcertClientSequencerManager::OnWorkspaceEndFrameCompleted()
 
 	for (const FString& SequenceObjectPath : PendingSequenceOpenEvents)
 	{
-		ApplyTransportOpenEvent(SequenceObjectPath);
+		ApplyOpenEvent(SequenceObjectPath);
 	}
 	PendingSequenceOpenEvents.Reset();
 
diff --git a/Engine/Plugins/Developer/Concert/ConcertSync/ConcertSyncClient/Source/ConcertSyncClient/Private/ConcertClientSequencerManager.h b/Engine/Plugins/Developer/Concert/ConcertSync/ConcertSyncClient/Source/ConcertSyncClient/Private/ConcertClientSequencerManager.h
--- a/Engine/Plugins/Developer/Concert/ConcertSync/ConcertSyncClient/Source/ConcertSyncClient/Private/ConcertClientSequencerManager.h
+++ b/Engine/Plugins/Developer/Concert/ConcertSync/ConcertSyncClient/Source/ConcertSyncClient/Private/ConcertClientSequencerManager.h
@@ -181,14 +181,14 @@ private:
 	 * @param InEventContext              The context for the current client session
 	 * @param InEvent                     The sequencer state sync event received from the server
 	 */
-	void OnSyncEvent(const FConcertSessionContext& InEventContext, const FConcertSequencerStateSyncEvent& InEvent);
+	void OnStateSyncEvent(const FConcertSessionContext& InEventContext, const FConcertSequencerStateSyncEvent& InEvent);
 
 	/**
-	 * Called on receipt of time adjustment event.
+	 * Called on receipt of an external time adjustment event from the server
 	 * @param InEventContext             The context for the current session.
 	 * @param InEvent                    The sequencer time adjustment event.
 	 */
-	void OnTimeAdjustmentEvent(const FConcertSessionContext& InEventContext, const FConcertSequencerTimeAdjustmentEvent& TimeAdjustmentEvent);
+	void OnTimeAdjustmentEvent(const FConcertSessionContext& InEventContext, const FConcertSequencerTimeAdjustmentEvent& InEvent);
 
 	/**
 	 * Called when the global time has been changed for the specified Sequencer
@@ -218,7 +218,7 @@ private:
 	 *
 	 * @param SequenceObjectPath	The sequence to open
 	 */
-	void ApplyTransportOpenEvent(const FString& SequenceObjectPath);
+	void ApplyOpenEvent(const FString& SequenceObjectPath);
 
 	/**
 	 * Create a new sequence player to be used in -game instances.
@@ -249,33 +249,41 @@ private:
 	void ApplyEventToPlayers(const FConcertSequencerState& PendingState);
 
 	/**
-	 * Apply a Sequencer Close Event to SequencePlayers
+	 * Apply a Sequencer close event to open players.
+	 * 
+	 * This is a helper function called by ApplyCloseEvent().
 	 *
-	 * @param PendingState	The pending state to apply
+	 * @param CloseEvent  The close event state
 	 */
-	void ApplyCloseToPlayers(const FConcertSequencerCloseEvent& InEvent);
+	void ApplyCloseEventToPlayers(const FConcertSequencerCloseEvent& CloseEvent);
 
 	/**
-	 * Apply CloseEvent to Open sequencers and players.
+	 * Apply a Sequencer close event to open Sequencers and players.
 	 *
-	 * @param PendingClose  The close event state.
+	 * @param CloseEvent  The close event state
 	 */
-	void ApplyTransportCloseEvent(const FConcertSequencerCloseEvent& PendingClose);
+	void ApplyCloseEvent(const FConcertSequencerCloseEvent& CloseEvent);
 
 	/**
 	 * Apply a sequencer time adjustment event.
+	 * 
+	 * @param TimeAdjustmentEvent  The time adjustment event state
 	 */
 	void ApplyTimeAdjustmentEvent(const FConcertSequencerTimeAdjustmentEvent& TimeAdjustmentEvent);
 
 	/**
 	 * Apply a sequencer time adjustment to players.
+	 * 
+	 * @param TimeAdjustmentEvent  The time adjustment event state
 	 */
-	void ApplyTimeAdjustmentPlayers(const FConcertSequencerTimeAdjustmentEvent& TimeAdjustmentEvent);
+	void ApplyTimeAdjustmentToPlayers(const FConcertSequencerTimeAdjustmentEvent& TimeAdjustmentEvent);
 
 	/**
 	 * Apply a sequencer time adjustment to sequencers.
+	 * 
+	 * @param TimeAdjustmentEvent  The time adjustment event state
 	 */
-	void ApplyTimeAdjustmentSequencers(const FConcertSequencerTimeAdjustmentEvent& TimeAdjustmentEvent);
+	void ApplyTimeAdjustmentToSequencers(const FConcertSequencerTimeAdjustmentEvent& TimeAdjustmentEvent);
 
 	/**
 	 * Gather all the currently open sequencer UIs that have the specified path as their root sequence
