Fix fully loaded state of the package on asset add. Since the asset is added the package is fully loaded the flags are incorrect because of the transaction is remotely.

diff --git a/Engine/Plugins/Developer/Concert/ConcertSync/ConcertSyncClient/Source/ConcertSyncClient/Private/ConcertClientPackageBridge.cpp b/Engine/Plugins/Developer/Concert/ConcertSync/ConcertSyncClient/Source/ConcertSyncClient/Private/ConcertClientPackageBridge.cpp
--- a/Engine/Plugins/Developer/Concert/ConcertSync/ConcertSyncClient/Source/ConcertSyncClient/Private/ConcertClientPackageBridge.cpp
+++ b/Engine/Plugins/Developer/Concert/ConcertSync/ConcertSyncClient/Source/ConcertSyncClient/Private/ConcertClientPackageBridge.cpp
@@ -269,6 +269,16 @@ void FConcertClientPackageBridge::HandleAssetAdded(UObject *Object)
 				IFileManager::Get().Delete(*PackageFilename);
 			}
 		}
+		else if (Object->IsPackageExternal() && Package->HasAnyPackageFlags(PKG_NewlyCreated))
+		{
+			// Save can fail if package was not fully loaded.  If we are in a one-file per actor case then we should
+			// clear flags to allow users to save it. We leave the package in a dirty state because we want the user to
+			// issue the save (if desired) and it is valid for save to fail on asset create when applied from a
+			// transaction.
+			//
+			Package->ClearPackageFlags(PKG_NewlyCreated);
+			Package->MarkAsFullyLoaded();
+		}
 	}
 
 	UE_LOG(LogConcert, Verbose, TEXT("Asset Added: %s"), *Package->GetName());
