ConcertSyncClient: fix static analysis warning about possible use of nullptr

diff --git a/Engine/Plugins/Developer/Concert/ConcertSync/ConcertSyncClient/Source/ConcertSyncClient/Private/ConcertClientSequencerManager.cpp b/Engine/Plugins/Developer/Concert/ConcertSync/ConcertSyncClient/Source/ConcertSyncClient/Private/ConcertClientSequencerManager.cpp
--- a/Engine/Plugins/Developer/Concert/ConcertSync/ConcertSyncClient/Source/ConcertSyncClient/Private/ConcertClientSequencerManager.cpp
+++ b/Engine/Plugins/Developer/Concert/ConcertSync/ConcertSyncClient/Source/ConcertSyncClient/Private/ConcertClientSequencerManager.cpp
@@ -245,9 +245,10 @@ bool FConcertClientSequencerManager::ShouldAlwaysCloseGameSequencerPlayer() cons
 void FConcertClientSequencerManager::OnSequencerClosed(TSharedRef<ISequencer> InSequencer)
 {
 	const UMovieSceneSequence* Sequence = InSequencer->GetRootMovieSceneSequence();
-	const FString SequenceObjectPath = Sequence->GetPathName();
+	const FString SequenceObjectPath = (Sequence != nullptr) ? Sequence->GetPathName() : FString();
 
-	UE_LOG(LogConcertSequencerSync, Verbose, TEXT("OnSequencerClosed: %s"), *SequenceObjectPath);
+	UE_LOG(LogConcertSequencerSync, Verbose, TEXT("OnSequencerClosed: %s"),
+		(!SequenceObjectPath.IsEmpty()) ? *SequenceObjectPath : TEXT("<NO SEQUENCE>"));
 
 	// Find the associated open sequencer index
 	int Index = 0;
