Add the persistent level for reload if we are open world level.

diff --git a/Engine/Plugins/Developer/Concert/ConcertSync/ConcertSyncClient/Source/ConcertSyncClient/Private/ConcertClientPackageManager.cpp b/Engine/Plugins/Developer/Concert/ConcertSync/ConcertSyncClient/Source/ConcertSyncClient/Private/ConcertClientPackageManager.cpp
--- a/Engine/Plugins/Developer/Concert/ConcertSync/ConcertSyncClient/Source/ConcertSyncClient/Private/ConcertClientPackageManager.cpp
+++ b/Engine/Plugins/Developer/Concert/ConcertSync/ConcertSyncClient/Source/ConcertSyncClient/Private/ConcertClientPackageManager.cpp
@@ -119,6 +119,21 @@ FConcertClientPackageManager::~FConcertClientPackageManager()
 		}
 	}
 
+	// The persistent level should always be reloaded when using external objects.
+	UWorld* CurrentWorld = ConcertSyncClientUtil::GetCurrentWorld();
+	if (CurrentWorld)
+	{
+		ULevel* PersistentLevel = CurrentWorld->PersistentLevel;
+		if (PersistentLevel && PersistentLevel->IsUsingExternalObjects())
+		{
+			const FName PackageName = PersistentLevel->GetPackage()->GetFName();
+			if (!PackagesPendingHotReload.Contains(PackageName))
+			{
+				PackagesPendingHotReload.Add(PackageName);
+			}
+		}
+	}
+
 	if (!IsEngineExitRequested())
 	{
 		// Hot reload after unregistering from most delegates to prevent events triggered by hot-reloading (such as asset deleted) to be recorded as transaction.
diff --git a/Engine/Plugins/Developer/Concert/ConcertSync/ConcertSyncClient/Source/ConcertSyncClient/Private/ConcertSyncClientUtil.cpp b/Engine/Plugins/Developer/Concert/ConcertSync/ConcertSyncClient/Source/ConcertSyncClient/Private/ConcertSyncClientUtil.cpp
--- a/Engine/Plugins/Developer/Concert/ConcertSync/ConcertSyncClient/Source/ConcertSyncClient/Private/ConcertSyncClientUtil.cpp
+++ b/Engine/Plugins/Developer/Concert/ConcertSync/ConcertSyncClient/Source/ConcertSyncClient/Private/ConcertSyncClientUtil.cpp
@@ -402,6 +402,10 @@
 		UPackage* ExistingPackage = FindPackage(nullptr, *PackageName.ToString());
 		if (ExistingPackage)
 		{
+			if (ExistingPackage->HasAnyPackageFlags(PKG_NewlyCreated))
+			{
+				ExistingPackage->ClearPackageFlags(PKG_NewlyCreated);
+			}
 			ExistingPackages.Add(ExistingPackage);
 		}
 	}
@@ -463,7 +467,7 @@
 		}
 	}
 
-	// Broadcast the eminent objects destruction (ex. tell BlueprintActionDatabase to release its reference(s) on Blueprint(s) right now)
+	// Broadcast the eminent objects destruction (ex. tell BlueprintActionDatabase to release its reference(s) on Blueprint(s) right now)
 	FEditorDelegates::OnAssetsPreDelete.Broadcast(ObjectsToPurge);
 
 	// Mark objects as purgeable.
