Add an option to allow users to make a file writable if source control is not enabled.

diff --git a/Engine/Plugins/Developer/Concert/ConcertApp/MultiUserClient/Source/MultiUserClient/Private/ConcertWorkspaceUI.cpp b/Engine/Plugins/Developer/Concert/ConcertApp/MultiUserClient/Source/MultiUserClient/Private/ConcertWorkspaceUI.cpp
--- a/Engine/Plugins/Developer/Concert/ConcertApp/MultiUserClient/Source/MultiUserClient/Private/ConcertWorkspaceUI.cpp
+++ b/Engine/Plugins/Developer/Concert/ConcertApp/MultiUserClient/Source/MultiUserClient/Private/ConcertWorkspaceUI.cpp
@@ -660,24 +660,30 @@ bool FConcertWorkspaceUI::PromptPersistSessionChanges()
 	FAsyncTaskNotification Notification(NotificationConfig);
 	FText NotificationSub;
 
-	TArray<FText> PersistFailures;
-	bool bSuccess = ClientWorkspacePin->PersistSessionChanges(PersistCmd.PackagesToPersist, &ISourceControlModule::Get().GetProvider(), &PersistFailures);
+	FPersistResult PersistResult = ClientWorkspacePin->PersistSessionChanges(
+		{PersistCmd.PackagesToPersist, &ISourceControlModule::Get().GetProvider(),PersistCmd.bShouldMakeFilesWritable});
+
+	bool bSuccess = PersistResult.PersistStatus == EPersistStatus::Success;
 	if (bSuccess)
 	{
 		bSuccess = SubmitChangelist(PersistCmd, NotificationSub);
 	}
 	else
 	{
-		NotificationSub = FText::Format(LOCTEXT("FailedPersistNotification", "Failed to persist session files. Reported {0} {0}|plural(one=error,other=errors)."), PersistFailures.Num());
+		NotificationSub = FText::Format(
+			LOCTEXT("FailedPersistNotification",
+					"Failed to persist session files. Reported {0} {0}|plural(one=error,other=errors)."), PersistResult.FailureReasons.Num());
 		FMessageLog ConcertLog("Concert");
-		for (const FText& Failure : PersistFailures)
+		for (const FText& Failure : PersistResult.FailureReasons)
 		{
 			ConcertLog.Error(Failure);
 		}
 	}
 
 	Notification.SetProgressText(LOCTEXT("SeeMessageLog", "See Message Log"));
-	Notification.SetComplete(bSuccess ? LOCTEXT("PersistChangeSuccessHeader", "Successfully Persisted Session Changes") : LOCTEXT("PersistChangeFailedHeader", "Failed to Persist Session Changes"), NotificationSub, bSuccess);
+	Notification.SetComplete(bSuccess ?
+							 LOCTEXT("PersistChangeSuccessHeader", "Successfully Persisted Session Changes")
+							 : LOCTEXT("PersistChangeFailedHeader", "Failed to Persist Session Changes"), NotificationSub, bSuccess);
 
 	return true; // Persisting.
 }
diff --git a/Engine/Plugins/Developer/Concert/ConcertApp/MultiUserClient/Source/MultiUserClient/Private/Widgets/SConcertSandboxPersistWidget.cpp b/Engine/Plugins/Developer/Concert/ConcertApp/MultiUserClient/Source/MultiUserClient/Private/Widgets/SConcertSandboxPersistWidget.cpp
--- a/Engine/Plugins/Developer/Concert/ConcertApp/MultiUserClient/Source/MultiUserClient/Private/Widgets/SConcertSandboxPersistWidget.cpp
+++ b/Engine/Plugins/Developer/Concert/ConcertApp/MultiUserClient/Source/MultiUserClient/Private/Widgets/SConcertSandboxPersistWidget.cpp
@@ -209,6 +209,19 @@ void SConcertSandboxPersistWidget::Construct(const FArguments& InArgs)
 						.Text(LOCTEXT("KeepFilesCheckedOut", "Keep Files Checked Out"))
 					]
 				]
+				+ SWrapBox::Slot()
+				.Padding(0.0f, 0.0f, 16.0f, 0.0f)
+				[
+					SNew(SCheckBox)
+					.OnCheckStateChanged(this, &SConcertSandboxPersistWidget::OnCheckStateChanged_MakeWritable)
+					.IsChecked(this, &SConcertSandboxPersistWidget::GetMakeWritable)
+					.ToolTipText(LOCTEXT("WarningMakeWritable", "Warning this may prevent accurate check-ins with source control."))
+					.IsEnabled_Lambda([this]() { return !CanCheckOut();} )
+					[
+						SNew(STextBlock)
+						.Text(LOCTEXT("MakeFilesWritable", "Make Files Writable"))
+					]
+				]
 			]
 			+SVerticalBox::Slot()
 			.AutoHeight()
@@ -256,6 +269,7 @@ FReply SConcertSandboxPersistWidget::OnKeyDown(const FGeometry & MyGeometry, con
 FConcertPersistCommand SConcertSandboxPersistWidget::GetPersistCommand() const
 {
 	FConcertPersistCommand Cmd;
+	Cmd.bShouldMakeFilesWritable = GetMakeWritable() == ECheckBoxState::Checked;
 	Cmd.bShouldSubmit = !IsSubmitDescriptionReadOnly();
 	Cmd.ChangelistDescription = SubmitDescriptionTextCtrl->GetText();
 	Cmd.PackagesToPersist.Reserve(ListViewItems.Num());
@@ -331,7 +345,7 @@ FReply SConcertSandboxPersistWidget::CancelClicked()
 	{
 		ParentWindowPin->RequestDestroyWindow();
 	}
-	
+
 	return FReply::Handled();
 }
 
@@ -371,6 +385,16 @@ bool SConcertSandboxPersistWidget::CanSubmitToSourceControl() const
 	return ISourceControlModule::Get().GetProvider().IsEnabled();
 }
 
+void SConcertSandboxPersistWidget::OnCheckStateChanged_MakeWritable(ECheckBoxState InState)
+{
+	MakeWritable = InState;
+}
+
+ECheckBoxState SConcertSandboxPersistWidget::GetMakeWritable() const
+{
+	return MakeWritable;
+}
+
 void SConcertSandboxPersistWidget::OnCheckStateChanged_KeepCheckedOut(ECheckBoxState InState)
 {
 	KeepFilesCheckedOut = InState;
diff --git a/Engine/Plugins/Developer/Concert/ConcertApp/MultiUserClient/Source/MultiUserClient/Private/Widgets/SConcertSandboxPersistWidget.h b/Engine/Plugins/Developer/Concert/ConcertApp/MultiUserClient/Source/MultiUserClient/Private/Widgets/SConcertSandboxPersistWidget.h
--- a/Engine/Plugins/Developer/Concert/ConcertApp/MultiUserClient/Source/MultiUserClient/Private/Widgets/SConcertSandboxPersistWidget.h
+++ b/Engine/Plugins/Developer/Concert/ConcertApp/MultiUserClient/Source/MultiUserClient/Private/Widgets/SConcertSandboxPersistWidget.h
@@ -21,6 +21,7 @@ struct FConcertPersistCommand
 	TArray<FString> FilesToPersist;
 	FText ChangelistDescription;
 	bool bShouldSubmit;
+	bool bShouldMakeFilesWritable;
 };
 
 /**
@@ -148,6 +149,12 @@ private:
 	/** Get the current state of the Keep Checked Out checkbox  */
 	ECheckBoxState GetKeepCheckedOut() const;
 
+	/** Called when the Make Writable checkbox is changed */
+	void OnCheckStateChanged_MakeWritable(ECheckBoxState InState);
+
+	/** Get the current state of the Make Writable checkbox  */
+	ECheckBoxState GetMakeWritable() const;
+
 	/** Check if Provider can checkout files */
 	bool CanCheckOut() const;
 
@@ -202,10 +209,12 @@ private:
 	/** State of the "Keep Files Checked Out" checkbox */
 	ECheckBoxState	KeepFilesCheckedOut;
 
+	/** State of the "Make Writable" checkbox */
+	ECheckBoxState	MakeWritable;
 
 	/** Specify which column to sort with */
 	FName SortByColumn;
 
 	/** Currently selected sorting mode */
 	EColumnSortMode::Type SortMode;
-};
\ No newline at end of file
+};
diff --git a/Engine/Plugins/Developer/Concert/ConcertSync/ConcertSyncClient/Source/ConcertSyncClient/Private/ConcertClientPackageManager.cpp b/Engine/Plugins/Developer/Concert/ConcertSync/ConcertSyncClient/Source/ConcertSyncClient/Private/ConcertClientPackageManager.cpp
--- a/Engine/Plugins/Developer/Concert/ConcertSync/ConcertSyncClient/Source/ConcertSyncClient/Private/ConcertClientPackageManager.cpp
+++ b/Engine/Plugins/Developer/Concert/ConcertSync/ConcertSyncClient/Source/ConcertSyncClient/Private/ConcertClientPackageManager.cpp
@@ -364,24 +364,24 @@ TOptional<FString> FConcertClientPackageManager::GetDeletedPackagePath(FName Pac
 	return {};
 }
 
-bool FConcertClientPackageManager::PersistSessionChanges(TArrayView<const FName> InPackagesToPersist, ISourceControlProvider* SourceControlProvider, TArray<FText>* OutFailureReasons)
+FPersistResult FConcertClientPackageManager::PersistSessionChanges(FPersistParameters InParam)
 {
 #if WITH_EDITOR
 	if (SandboxPlatformFile)
 	{
 		// Transform all the package names into actual filenames
 		TArray<FString, TInlineAllocator<8>> FilesToPersist;
-		for (const FName& PackageName : InPackagesToPersist)
+		for (const FName& PackageName : InParam.PackagesToPersist)
 		{
 			if (TOptional<FString> ValidPath = GetValidPackageSessionPath(PackageName))
 			{
 				FilesToPersist.Add(MoveTemp(ValidPath.GetValue()));
 			}
 		}
-		return SandboxPlatformFile->PersistSandbox(FilesToPersist, SourceControlProvider, OutFailureReasons);
+		return SandboxPlatformFile->PersistSandbox(FilesToPersist, MoveTemp(InParam));
 	}
 #endif
-	return false;
+	return {};
 }
 
 bool FConcertClientPackageManager::ApplyPackageFilters(const FConcertPackageInfo& InPackageInfo) const
@@ -540,7 +540,7 @@ void FConcertClientPackageManager::HandleLocalPackageEvent(const FConcertPackage
 	// if the package data has been filtered out of the session persist it immediately from the sandbox
 	else
 	{
-		PersistSessionChanges(MakeArrayView(&PackageInfo.PackageName, 1), &ISourceControlModule::Get().GetProvider());
+		PersistSessionChanges({MakeArrayView(&PackageInfo.PackageName, 1), &ISourceControlModule::Get().GetProvider()});
 	}
 }
 
diff --git a/Engine/Plugins/Developer/Concert/ConcertSync/ConcertSyncClient/Source/ConcertSyncClient/Private/ConcertClientPackageManager.h b/Engine/Plugins/Developer/Concert/ConcertSync/ConcertSyncClient/Source/ConcertSyncClient/Private/ConcertClientPackageManager.h
--- a/Engine/Plugins/Developer/Concert/ConcertSync/ConcertSyncClient/Source/ConcertSyncClient/Private/ConcertClientPackageManager.h
+++ b/Engine/Plugins/Developer/Concert/ConcertSync/ConcertSyncClient/Source/ConcertSyncClient/Private/ConcertClientPackageManager.h
@@ -6,6 +6,8 @@
 #include "IConcertSessionHandler.h"
 #include "IConcertClientPackageBridge.h"
 
+#include "ConcertClientPersistData.h"
+
 class FConcertSyncClientLiveSession;
 class FConcertSandboxPlatformFile;
 class ISourceControlProvider;
@@ -98,7 +100,7 @@ public:
 	/**
 	 * Persist the session changes from the package name list and prepare it for source control submission.
 	 */
-	bool PersistSessionChanges(TArrayView<const FName> InPackagesToPersist, ISourceControlProvider* SourceControlProvider, TArray<FText>* OutFailureReasons = nullptr);
+	FPersistResult PersistSessionChanges(FPersistParameters InParam);
 
 	/**
 	 * Called when a package is too big to be handled by the system.
diff --git a/Engine/Plugins/Developer/Concert/ConcertSync/ConcertSyncClient/Source/ConcertSyncClient/Private/ConcertClientWorkspace.cpp b/Engine/Plugins/Developer/Concert/ConcertSync/ConcertSyncClient/Source/ConcertSyncClient/Private/ConcertClientWorkspace.cpp
--- a/Engine/Plugins/Developer/Concert/ConcertSync/ConcertSyncClient/Source/ConcertSyncClient/Private/ConcertClientWorkspace.cpp
+++ b/Engine/Plugins/Developer/Concert/ConcertSync/ConcertSyncClient/Source/ConcertSyncClient/Private/ConcertClientWorkspace.cpp
@@ -306,31 +306,35 @@ TOptional<FString> FConcertClientWorkspace::GetValidPackageSessionPath(FName Pac
 	return {};
 }
 
-bool FConcertClientWorkspace::PersistSessionChanges(TArrayView<const FName> InPackagesToPersist, ISourceControlProvider* SourceControlProvider, TArray<FText>* OutFailureReasons)
+
+FPersistResult FConcertClientWorkspace::PersistSessionChanges(FPersistParameters InParam)
 {
 	bool bSuccess = false;
 #if WITH_EDITOR
 	if (PackageManager)
 	{
-		for (const FName& PackageName : InPackagesToPersist)
+		TArray<FName> PackageNames;
+		for (const FName& PackageName : InParam.PackagesToPersist)
 		{
 			SaveLiveTransactionsToPackage(PackageName);
+			PackageNames.Add(PackageName);
 		}
 
-		bSuccess = PackageManager->PersistSessionChanges(InPackagesToPersist, SourceControlProvider, OutFailureReasons);
-
+		FPersistResult Result = PackageManager->PersistSessionChanges(MoveTemp(InParam));
 		// if we successfully persisted the files, record persist events for them in the db
-		if (bSuccess)
+		if (Result.PersistStatus == EPersistStatus::Success)
 		{
 			int64 PersistEventId = 0;
-			for (const FName& PackageName : InPackagesToPersist)
+			for (const FName& PackageName : PackageNames)
 			{
-				LiveSession->GetSessionDatabase().AddPersistEventForHeadRevision(PackageName, PersistEventId);
+				LiveSession->GetSessionDatabase().AddPersistEventForHeadRevision(
+					PackageName, PersistEventId);
 			}
 		}
+		return Result;
 	}
 #endif
-	return bSuccess;
+	return {};
 }
 
 bool FConcertClientWorkspace::HasLiveTransactionSupport(UPackage* InPackage) const
diff --git a/Engine/Plugins/Developer/Concert/ConcertSync/ConcertSyncClient/Source/ConcertSyncClient/Private/ConcertClientWorkspace.h b/Engine/Plugins/Developer/Concert/ConcertSync/ConcertSyncClient/Source/ConcertSyncClient/Private/ConcertClientWorkspace.h
--- a/Engine/Plugins/Developer/Concert/ConcertSync/ConcertSyncClient/Source/ConcertSyncClient/Private/ConcertClientWorkspace.h
+++ b/Engine/Plugins/Developer/Concert/ConcertSync/ConcertSyncClient/Source/ConcertSyncClient/Private/ConcertClientWorkspace.h
@@ -44,7 +44,7 @@ public:
 	virtual bool HasSessionChanges() const override;
 	virtual TArray<FName> GatherSessionChanges(bool IgnorePersisted = true) override;
 	virtual TOptional<FString> GetValidPackageSessionPath(FName PackageName) const override;
-	virtual bool PersistSessionChanges(TArrayView<const FName> InPackagesToPersist, ISourceControlProvider* SourceControlProvider, TArray<FText>* OutFailureReasons = nullptr) override;
+	virtual FPersistResult PersistSessionChanges(FPersistParameters InParam) override;
 	virtual bool HasLiveTransactionSupport(UPackage* InPackage) const override;
 	virtual bool ShouldIgnorePackageDirtyEvent(class UPackage* InPackage) const override;
 	virtual bool FindTransactionEvent(const int64 TransactionEventId, FConcertSyncTransactionEvent& OutTransactionEvent, const bool bMetaDataOnly) const override;
diff --git a/Engine/Plugins/Developer/Concert/ConcertSync/ConcertSyncClient/Source/ConcertSyncClient/Private/ConcertSandboxPlatformFile.cpp b/Engine/Plugins/Developer/Concert/ConcertSync/ConcertSyncClient/Source/ConcertSyncClient/Private/ConcertSandboxPlatformFile.cpp
--- a/Engine/Plugins/Developer/Concert/ConcertSync/ConcertSyncClient/Source/ConcertSyncClient/Private/ConcertSandboxPlatformFile.cpp
+++ b/Engine/Plugins/Developer/Concert/ConcertSync/ConcertSyncClient/Source/ConcertSyncClient/Private/ConcertSandboxPlatformFile.cpp
@@ -731,119 +731,128 @@ FString FConcertSandboxPlatformFile::ConvertToAbsolutePathForExternalAppForWrite
 	return LowerLevel->ConvertToAbsolutePathForExternalAppForWrite(*ResolvedPath.GetNonSandboxPath());
 }
 
-bool FConcertSandboxPlatformFile::PersistSandbox(TArrayView<const FString> InFiles, ISourceControlProvider* SourceControlProvider, TArray<FText>* OutFailureReasons)
+
+bool FConcertSandboxPlatformFile::CopyFileWithSCC(
+	const TCHAR* InToFilename,
+	const TCHAR* InFromFilename,
+	const FPersistParameters& InParam,
+	FPersistResult& OutResult)
 {
-	// Lambda copy file from the sandbox
-	auto CopyFileWithSCC = [this, SourceControlProvider, OutFailureReasons](const TCHAR* InToFilename, const TCHAR* InFromFilename)
-	{
-		// If this file maps to a package then we need to flush its linker so that we can overwrite the file on disk
-		ConcertSandboxPlatformFileUtil::FlushPackageFile(InToFilename);
+	// If this file maps to a package then we need to flush its linker so that we can overwrite the file on disk
+	ConcertSandboxPlatformFileUtil::FlushPackageFile(InToFilename);
 
-		// Get the source control state of the destination file
-		FSourceControlStatePtr ToFileSCCState;
-		if (SourceControlProvider && SourceControlProvider->IsEnabled())
-		{
-			ToFileSCCState = SourceControlProvider->GetState(InToFilename, EStateCacheUsage::ForceUpdate);
-		}
+	ISourceControlProvider* SourceControlProvider = InParam.SourceControlProvider;
+	// Get the source control state of the destination file
+	FSourceControlStatePtr ToFileSCCState;
+	if (SourceControlProvider && SourceControlProvider->IsEnabled())
+	{
+		ToFileSCCState = SourceControlProvider->GetState(InToFilename, EStateCacheUsage::ForceUpdate);
+	}
 
-		// We don't need to do anything with source control if the file is already checked-out or added
-		const bool bRequiresSCCAction = ToFileSCCState && !ToFileSCCState->IsCheckedOut() && !ToFileSCCState->IsAdded();
+	// We don't need to do anything with source control if the file is already checked-out or added
+	const bool bRequiresSCCAction = ToFileSCCState && !ToFileSCCState->IsCheckedOut() && !ToFileSCCState->IsAdded();
 
-		// If the file can be checked-out, do so now
-		if (bRequiresSCCAction && ToFileSCCState->IsSourceControlled())
+	// If the file can be checked-out, do so now
+	if (bRequiresSCCAction && ToFileSCCState->IsSourceControlled())
+	{
+		// the static analysis tool is not able to see that `SourceControlProvider` won't be null if `bRequiresSCCAction` is true
+		CA_ASSUME(SourceControlProvider != nullptr);
+		if (ToFileSCCState->CanCheckout() && SourceControlProvider->UsesCheckout())
 		{
-			// the static analysis tool is not able to see that `SourceControlProvider` won't be null if `bRequiresSCCAction` is true
-			CA_ASSUME(SourceControlProvider != nullptr);
-			if (ToFileSCCState->CanCheckout() && SourceControlProvider->UsesCheckout())
-			{
-				TArray<FString> FilesToBeCheckedOut;
-				FilesToBeCheckedOut.Add(InToFilename);
+			TArray<FString> FilesToBeCheckedOut;
+			FilesToBeCheckedOut.Add(InToFilename);
 
-				if (SourceControlProvider->Execute(ISourceControlOperation::Create<FCheckOut>(), FilesToBeCheckedOut) != ECommandResult::Succeeded)
-				{
-					FText Failure = FText::Format(NSLOCTEXT("Concert", "CheckoutFileFailure", "Failed to check-out file '{0}' from source control when persiting sandbox state!"), FText::FromString(InToFilename));
-					UE_LOG(LogConcert, Warning, TEXT("%s"),  *Failure.ToString());
-					if (OutFailureReasons)
-					{
-						OutFailureReasons->Add(MoveTemp(Failure));
-					}
-					return false;
-				}
-			}
-			else
+			if (SourceControlProvider->Execute(ISourceControlOperation::Create<FCheckOut>(), FilesToBeCheckedOut) != ECommandResult::Succeeded)
 			{
-				FText Failure = FText::Format(NSLOCTEXT("Concert", "CanCheckoutFileFailure", "Can't check-out file '{0}' from source control when persiting sandbox state!"), FText::FromString(InToFilename));
-				UE_LOG(LogConcert, Warning, TEXT("%s"), *Failure.ToString());
-				if (OutFailureReasons)
-				{
-					OutFailureReasons->Add(MoveTemp(Failure));
-				}
+				FText Failure = FText::Format(
+					NSLOCTEXT("Concert", "CheckoutFileFailure", "Failed to check-out file '{0}' from source control when persiting sandbox state!"),
+					FText::FromString(InToFilename));
+				UE_LOG(LogConcert, Warning, TEXT("%s"),  *Failure.ToString());
+				OutResult.FailureReasons.Add(MoveTemp(Failure));
 				return false;
 			}
 		}
-
-		// Copy the on-disk sandbox file
-		FString ToFileDir = FPaths::GetPath(InToFilename);
-		if (!CreateDirectoryTree(*ToFileDir) || !LowerLevel->CopyFile(InToFilename, InFromFilename))
+		else
 		{
-			FText Failure = FText::Format(NSLOCTEXT("Concert", "CopyFileFailure", "Failed to copy file '{0}' (from '{1}') when persiting sandbox state!"), FText::FromString(InToFilename), FText::FromString(InFromFilename));
+			FText Failure = FText::Format(
+				NSLOCTEXT("Concert", "CanCheckoutFileFailure", "Can't check-out file '{0}' from source control when persiting sandbox state!"),
+				FText::FromString(InToFilename));
 			UE_LOG(LogConcert, Warning, TEXT("%s"), *Failure.ToString());
-			if (OutFailureReasons)
-			{
-				OutFailureReasons->Add(MoveTemp(Failure));
-			}
+			OutResult.FailureReasons.Add(MoveTemp(Failure));
 			return false;
 		}
-
-		// If the file is new, add it to source control now
-		if (bRequiresSCCAction && !ToFileSCCState->IsSourceControlled())
+	}
+	else if (InParam.bShouldMakeWritableIfNoSourceControl)
+	{
+		if (!LowerLevel->SetReadOnly(InToFilename, false))
 		{
-			TArray<FString> FilesToBeAdded;
-			FilesToBeAdded.Add(InToFilename);
-
-			if (SourceControlProvider->Execute(ISourceControlOperation::Create<FMarkForAdd>(), FilesToBeAdded) != ECommandResult::Succeeded)
-			{
-				FText Failure = FText::Format(NSLOCTEXT("Concert", "AddFileFailure", "Failed to add file '{0}' to source control when persiting sandbox state!"), FText::FromString(InToFilename));
-				UE_LOG(LogConcert, Warning, TEXT("%s"), *Failure.ToString());
-				if (OutFailureReasons)
-				{
-					OutFailureReasons->Add(MoveTemp(Failure));
-				}
-				return false;
-			}
+			FText Failure = FText::Format(
+				NSLOCTEXT("Concert", "CanMakeWritable", "Can't make file '{0}' writable when persiting sandbox state!"),
+				FText::FromString(InToFilename));
+			UE_LOG(LogConcert, Warning, TEXT("%s"), *Failure.ToString());
+			OutResult.FailureReasons.Add(MoveTemp(Failure));
+			return false;
 		}
-		return true;
-	};
-	// Lambda to delete files removed from the sandbox
-	auto DeleteFileWithSCC = [this, SourceControlProvider, OutFailureReasons](const TCHAR* InFilename)
+	}
+
+	// Copy the on-disk sandbox file
+	FString ToFileDir = FPaths::GetPath(InToFilename);
+	if (!CreateDirectoryTree(*ToFileDir) || !LowerLevel->CopyFile(InToFilename, InFromFilename))
+	{
+		FText Failure = FText::Format(
+			NSLOCTEXT("Concert", "CopyFileFailure", "Failed to copy file '{0}' (from '{1}') when persiting sandbox state!"),
+			FText::FromString(InToFilename), FText::FromString(InFromFilename));
+		UE_LOG(LogConcert, Warning, TEXT("%s"), *Failure.ToString());
+		OutResult.FailureReasons.Add(MoveTemp(Failure));
+		return false;
+	}
+
+		// If the file is new, add it to source control now
+	if (bRequiresSCCAction && !ToFileSCCState->IsSourceControlled())
 	{
-		// If this file maps to a package then we need to flush its linker so that we can remove the file from disk
-		ConcertSandboxPlatformFileUtil::FlushPackageFile(InFilename);
+		TArray<FString> FilesToBeAdded;
+		FilesToBeAdded.Add(InToFilename);
 
-		// Get the source control state of the file
-		FSourceControlStatePtr FileSCCState;
-		if (SourceControlProvider && SourceControlProvider->IsEnabled())
+		if (SourceControlProvider->Execute(ISourceControlOperation::Create<FMarkForAdd>(), FilesToBeAdded) != ECommandResult::Succeeded)
 		{
-			FileSCCState = SourceControlProvider->GetState(InFilename, EStateCacheUsage::ForceUpdate);
+			FText Failure = FText::Format(
+				NSLOCTEXT("Concert", "AddFileFailure", "Failed to add file '{0}' to source control when persiting sandbox state!"),
+				FText::FromString(InToFilename));
+			UE_LOG(LogConcert, Warning, TEXT("%s"), *Failure.ToString());
+			OutResult.FailureReasons.Add(MoveTemp(Failure));
+			return false;
 		}
+	}
+	return true;
+}
 
-		// Try and let source control remove the file first
-		if (FileSCCState && FileSCCState->IsSourceControlled())
-		{
-			const bool bAdded = FileSCCState->IsAdded();
+/** Delete file using source control */
+bool FConcertSandboxPlatformFile::DeleteFileWithSCC(const TCHAR* InFilename, const FPersistParameters& InParam, FPersistResult& OutResult)
+{
+	// If this file maps to a package then we need to flush its linker so that we can remove the file from disk
+	ConcertSandboxPlatformFileUtil::FlushPackageFile(InFilename);
 
-			if (bAdded || FileSCCState->IsCheckedOut())
+	ISourceControlProvider* SourceControlProvider = InParam.SourceControlProvider;
+	// Get the source control state of the file
+	FSourceControlStatePtr FileSCCState;
+	if (SourceControlProvider && SourceControlProvider->IsEnabled())
+	{
+		FileSCCState = SourceControlProvider->GetState(InFilename, EStateCacheUsage::ForceUpdate);
+	}
+
+	// Try and let source control remove the file first
+	if (FileSCCState && FileSCCState->IsSourceControlled())
+	{
+		const bool bAdded = FileSCCState->IsAdded();
+
+		if (bAdded || FileSCCState->IsCheckedOut())
+		{
+			if (SourceControlProvider->Execute(ISourceControlOperation::Create<FRevert>(), InFilename) != ECommandResult::Succeeded)
 			{
-				if (SourceControlProvider->Execute(ISourceControlOperation::Create<FRevert>(), InFilename) != ECommandResult::Succeeded)
-				{
-					FText Failure = FText::Format(NSLOCTEXT("Concert", "RevertFileFailure", "Failed to revert file '{0}' to source control when persisting sandbox state!"), FText::FromString(InFilename));
-					UE_LOG(LogConcert, Warning, TEXT("%s"), *Failure.ToString());
-					if (OutFailureReasons)
-					{
-						OutFailureReasons->Add(MoveTemp(Failure));
-					}
-					return false;
-				}
+				FText Failure = FText::Format(NSLOCTEXT("Concert", "RevertFileFailure", "Failed to revert file '{0}' to source control when persisting sandbox state!"), FText::FromString(InFilename));
+				UE_LOG(LogConcert, Warning, TEXT("%s"), *Failure.ToString());
+				OutResult.FailureReasons.Add(MoveTemp(Failure));
+				return false;
 			}
 
 			if (!bAdded)
@@ -852,44 +861,36 @@ bool FConcertSandboxPlatformFile::PersistSandbox(TArrayView<const FString> InFil
 				{
 					FText Failure = FText::Format(NSLOCTEXT("Concert", "DeleteSCCFileFailure", "Failed to delete file '{0}' from source control when persiting sandbox state!"), FText::FromString(InFilename));
 					UE_LOG(LogConcert, Warning, TEXT("%s"), *Failure.ToString());
-					if (OutFailureReasons)
-					{
-						OutFailureReasons->Add(MoveTemp(Failure));
-					}
+					OutResult.FailureReasons.Add(MoveTemp(Failure));
 					return false;
 				}
 			}
 		}
+	}
 
-		// Delete file if it still exists
-		if (LowerLevel->FileExists(InFilename) && !LowerLevel->DeleteFile(InFilename))
-		{
-			FText Failure = FText::Format(NSLOCTEXT("Concert", "DeleteFileFailure", "Failed to delete file '{0}' when persiting sandbox state!"), FText::FromString(InFilename));
-			UE_LOG(LogConcert, Warning, TEXT("%s"), *Failure.ToString());
-			if (OutFailureReasons)
-			{
-				OutFailureReasons->Add(MoveTemp(Failure));
-			}
-			return false;
-		}
-		return true;
-	};
+	// Delete file if it still exists
+	if (LowerLevel->FileExists(InFilename) && !LowerLevel->DeleteFile(InFilename))
+	{
+		FText Failure = FText::Format(NSLOCTEXT("Concert", "DeleteFileFailure", "Failed to delete file '{0}' when persiting sandbox state!"), FText::FromString(InFilename));
+		UE_LOG(LogConcert, Warning, TEXT("%s"), *Failure.ToString());
+		OutResult.FailureReasons.Add(MoveTemp(Failure));
+		return false;
+	}
+	return true;
+}
 
+FPersistResult FConcertSandboxPlatformFile::PersistSandbox(TArrayView<const FString> InFiles, FPersistParameters InParams)
+{
+	FPersistResult OutResult;
 	// We need to disable the sandbox while we do this
 	bool bFullOperationSuccess = true;
 	TGuardValue<TAtomic<bool>, bool> DisableSandboxGuard(bSandboxEnabled, false);
 	for (const FString& File : InFiles)
 	{
 		FConcertSandboxPlatformFilePath FilePath = ToSandboxPath(File, true);
-		bool bSuccess = true;
-		if (IsPathDeleted(FilePath))
-		{
-			bSuccess = DeleteFileWithSCC(*FilePath.GetNonSandboxPath());
-		}
-		else
-		{
-			bSuccess = CopyFileWithSCC(*FilePath.GetNonSandboxPath(), *FilePath.GetSandboxPath());
-		}
+		bool bSuccess = IsPathDeleted(FilePath) ?
+			DeleteFileWithSCC(*FilePath.GetNonSandboxPath(), InParams, OutResult) :
+			CopyFileWithSCC(*FilePath.GetNonSandboxPath(), *FilePath.GetSandboxPath(), InParams, OutResult);
 
 		// if the file operation was successful, mark the file as persisted
 		if (bSuccess)
@@ -899,7 +900,8 @@ bool FConcertSandboxPlatformFile::PersistSandbox(TArrayView<const FString> InFil
 
 		bFullOperationSuccess &= bSuccess;
 	}
-	return bFullOperationSuccess;
+	OutResult.PersistStatus = bFullOperationSuccess ? EPersistStatus::Success : EPersistStatus::Failure;
+	return OutResult;
 }
 
 void FConcertSandboxPlatformFile::DiscardSandbox(TArray<FName>& OutPackagesPendingHotReload, TArray<FName>& OutPackagesPendingPurge)
diff --git a/Engine/Plugins/Developer/Concert/ConcertSync/ConcertSyncClient/Source/ConcertSyncClient/Private/ConcertSandboxPlatformFile.h b/Engine/Plugins/Developer/Concert/ConcertSync/ConcertSyncClient/Source/ConcertSyncClient/Private/ConcertSandboxPlatformFile.h
--- a/Engine/Plugins/Developer/Concert/ConcertSync/ConcertSyncClient/Source/ConcertSyncClient/Private/ConcertSandboxPlatformFile.h
+++ b/Engine/Plugins/Developer/Concert/ConcertSync/ConcertSyncClient/Source/ConcertSyncClient/Private/ConcertSandboxPlatformFile.h
@@ -4,6 +4,7 @@
 
 #include "CoreMinimal.h"
 #include "GenericPlatform/GenericPlatformFile.h"
+#include "ConcertClientPersistData.h"
 
 class ISourceControlProvider;
 
@@ -150,16 +151,15 @@ public:
 	virtual FString ConvertToAbsolutePathForExternalAppForRead(const TCHAR* Filename) override;
 	virtual FString ConvertToAbsolutePathForExternalAppForWrite(const TCHAR* Filename) override;
 
-	/** 
-	 *	Persist the file list from the sandbox state onto the real files.
-	 *	Will mark files for which the operation was succesful as persisted.
+	/**
+	 *	Persist the file list from the sandbox state onto the real files Will mark files for which the operation was
+	 *	succesful as persisted.
 	 *
-	 *	@param InFiles					The files to persist from the sandbox
-	 *	@param SourceControlProvider	Optional pointer to the source control provider to stage/checkout/add the persisted files
-	 *	@param OutFailureReasons		Optional pointer to an array of text containing failure reasons if the operation fails.
-	 *	@return true if the operation was successful
+	 *	@param InFiles	List of files names to persist as full paths.
+	 *	@param InParams	Parameters controlling the persist options specified by the user.
+	 *	@return Returns a persist result object that contains result status.
 	 */
-	bool PersistSandbox(TArrayView<const FString> InFiles, ISourceControlProvider* SourceControlProvider = nullptr, TArray<FText>* OutFailureReasons = nullptr);
+	FPersistResult PersistSandbox(TArrayView<const FString> InFiles, FPersistParameters InParams);
 
 	/** 
 	 *	Discard the sandbox state
@@ -265,6 +265,12 @@ private:
 	void OnDirectoryChanged(const TArray<FFileChangeData>& FileChanges, FConcertSandboxPlatformFilePath MountPath);
 #endif
 
+	/** Copy file using source control (if enabled) */
+	bool CopyFileWithSCC(const TCHAR* InToFilename, const TCHAR* InFromFilename, const FPersistParameters& InParam, FPersistResult& OutResult);
+
+	/** Delete file using source control (if enabled) */
+	bool DeleteFileWithSCC(const TCHAR* InFilename, const FPersistParameters& InParam, FPersistResult& OutResult);
+
 	/** Root path of this sandbox */
 	const FString SandboxRootPath;
 
diff --git a/Engine/Plugins/Developer/Concert/ConcertSync/ConcertSyncClient/Source/ConcertSyncClient/Private/ConcertSyncClient.cpp b/Engine/Plugins/Developer/Concert/ConcertSync/ConcertSyncClient/Source/ConcertSyncClient/Private/ConcertSyncClient.cpp
--- a/Engine/Plugins/Developer/Concert/ConcertSync/ConcertSyncClient/Source/ConcertSyncClient/Private/ConcertSyncClient.cpp
+++ b/Engine/Plugins/Developer/Concert/ConcertSync/ConcertSyncClient/Source/ConcertSyncClient/Private/ConcertSyncClient.cpp
@@ -110,7 +110,7 @@ void FConcertSyncClient::PersistAllSessionChanges()
 	if (Workspace)
 	{
 		TArray<FName> SessionChanges = Workspace->GatherSessionChanges();
-		Workspace->PersistSessionChanges(SessionChanges, SourceControlProxy.Get());
+		Workspace->PersistSessionChanges({SessionChanges, SourceControlProxy.Get()});
 	}
 #endif
 }
diff --git a/Engine/Plugins/Developer/Concert/ConcertSync/ConcertSyncClient/Source/ConcertSyncClient/Public/ConcertClientPersistData.h b/Engine/Plugins/Developer/Concert/ConcertSync/ConcertSyncClient/Source/ConcertSyncClient/Public/ConcertClientPersistData.h
new file mode 100644
--- /dev/null
+++ b/Engine/Plugins/Developer/Concert/ConcertSync/ConcertSyncClient/Source/ConcertSyncClient/Public/ConcertClientPersistData.h
@@ -0,0 +1,43 @@
+// Copyright Epic Games, Inc. All Rights Reserved.
+
+#pragma once
+
+#include "Internationalization/Text.h"
+#include "UObject/NameTypes.h"
+#include "Containers/Array.h"
+
+class ISourceControlProvider;
+
+/** Input parameters for persistance of an existing session. */
+struct FPersistParameters
+{
+	/** Name of the packages to persist.*/
+	TArrayView<const FName> PackagesToPersist;
+
+	/** Pointer to the source control provider. */
+	ISourceControlProvider* SourceControlProvider;
+
+	/** Boolean to control if we will force the file writable when no source control provider is active. */
+	bool bShouldMakeWritableIfNoSourceControl = false;
+};
+
+/** Possible outcomes from a session persist. */
+enum class EPersistStatus
+{
+	/** The packages were persisted correctly. */
+	Success,
+	/** We failed to persist the packages. */
+	Failure,
+	/** We are not allowed to persist. */
+	NotAllowed
+};
+
+/** Result of a persist.*/
+struct FPersistResult
+{
+	/** List of reasons for persist failure. */
+	TArray<FText> FailureReasons;
+
+	/** Success or failure of the session persist. */
+	EPersistStatus PersistStatus = EPersistStatus::NotAllowed;
+};
diff --git a/Engine/Plugins/Developer/Concert/ConcertSync/ConcertSyncClient/Source/ConcertSyncClient/Public/IConcertClientWorkspace.h b/Engine/Plugins/Developer/Concert/ConcertSync/ConcertSyncClient/Source/ConcertSyncClient/Public/IConcertClientWorkspace.h
--- a/Engine/Plugins/Developer/Concert/ConcertSync/ConcertSyncClient/Source/ConcertSyncClient/Public/IConcertClientWorkspace.h
+++ b/Engine/Plugins/Developer/Concert/ConcertSync/ConcertSyncClient/Source/ConcertSyncClient/Public/IConcertClientWorkspace.h
@@ -8,6 +8,8 @@
 #include "ConcertWorkspaceMessages.h"
 #include "ConcertSyncSessionTypes.h"
 
+#include "ConcertClientPersistData.h"
+
 class ISourceControlProvider;
 class IConcertClientSession;
 class IConcertClientDataStore;
@@ -81,7 +83,7 @@ public:
 	virtual TOptional<FString> GetValidPackageSessionPath(FName PackageName) const = 0;
 
 	/** Persist the session changes from the package list and prepare it for source control submission */
-	virtual bool PersistSessionChanges(TArrayView<const FName> InPackageToPersist, ISourceControlProvider* SourceControlProvider, TArray<FText>* OutFailureReasonMap = nullptr) = 0;
+	virtual FPersistResult PersistSessionChanges(FPersistParameters InParam) = 0;
 
 	/**
 	 * Get Activities from the session.
