ConcertSyncClient: use FSoftObjectPath::ResolveObject() when finding object outers in ConcertSyncClientUtil::GetObject()

In WorldPartition levels, the path of the package to which an object is outered can vary
depending on the runtime context, e.g. editor vs. PIE vs. -game, etc. When a transaction
that affects an object is recorded on one client in a Multi-User session (for example, an editor),
the path to that object may not be the same on another client (for example, an nDisplay node
which runs using -game). This can manifest as object changes on one client not having an
effect on other clients in the session.

With this change, we use FSoftObjectPath's object resolution to find the object referenced
by the path rather than looking specifically for an object with that path using StaticFindObject().
This ensures that we pass through the ULevel's subobject resolution, in which there's a custom
resolver for subobjects inside a UWorldPartition.

Note that this provides a solution for resolving an unpartitioned path to a partitioned path.
A different solution would be needed for the reverse direction, to get an unpartitioned path from
a partitioned one, but that's not currently a situation that arises in Multi-User.

diff --git a/Engine/Plugins/Developer/Concert/ConcertSync/ConcertSyncClient/Source/ConcertSyncClient/Private/ConcertSyncClientUtil.cpp b/Engine/Plugins/Developer/Concert/ConcertSync/ConcertSyncClient/Source/ConcertSyncClient/Private/ConcertSyncClientUtil.cpp
--- a/Engine/Plugins/Developer/Concert/ConcertSync/ConcertSyncClient/Source/ConcertSyncClient/Private/ConcertSyncClientUtil.cpp
+++ b/Engine/Plugins/Developer/Concert/ConcertSync/ConcertSyncClient/Source/ConcertSyncClient/Private/ConcertSyncClientUtil.cpp
@@ -149,8 +149,15 @@
 		}
 	};
 
-	// Find the outer for the existing object
-	if (UObject* ExistingObjectOuter = StaticFindObject(UObject::StaticClass(), nullptr, *ObjectOuterPathToFind.ToString()))
+	// Find the outer for the existing object.
+	// Note that we use FSoftObjectPath::ResolveObject() here to ensure that if
+	// world partitioning is involved, we're able to resolve a non-partitioned
+	// path into an object with a partitioned path (e.g. an editor path to a
+	// "-game"/nDisplay path).
+	// TODO: If a case arises where we need to go the other direction and get
+	// an object with a non-partitioned path from a partitioned path, a
+	// different mechanism for that would be needed here.
+	if (UObject* ExistingObjectOuter = FSoftObjectPath(ObjectOuterPathToFind).ResolveObject())
 	{
 		// We need the object class to find or create the object
 		if (UClass* ObjectClass = FindOrLoadClass(InObjectId.ObjectClassPathName))
@@ -197,8 +204,10 @@
 		}
 	}
 
-	// Find the outer for the new object
-	if (UObject* NewObjectOuter = StaticFindObject(UObject::StaticClass(), nullptr, *ObjectOuterPathToCreate.ToString()))
+	// Find the outer for the new object.
+	// As above, we use FSoftObjectPath::ResolveObject() here to account for
+	// the possibility of world partitioning.
+	if (UObject* NewObjectOuter = FSoftObjectPath(ObjectOuterPathToCreate).ResolveObject())
 	{
 		// We need the object class to find or create the object
 		if (UClass* ObjectClass = FindOrLoadClass(InObjectId.ObjectClassPathName))
@@ -467,7 +476,7 @@
 		}
 	}
 
-	// Broadcast the eminent objects destruction (ex. tell BlueprintActionDatabase to release its reference(s) on Blueprint(s) right now)
+	// Broadcast the eminent objects destruction (ex. tell BlueprintActionDatabase to release its reference(s) on Blueprint(s) right now)
 	FEditorDelegates::OnAssetsPreDelete.Broadcast(ObjectsToPurge);
 
 	// Mark objects as purgeable.
