Ensure the file exists before changing the read only flag.

diff --git a/Engine/Plugins/Developer/Concert/ConcertSync/ConcertSyncClient/Source/ConcertSyncClient/Private/ConcertSandboxPlatformFile.cpp b/Engine/Plugins/Developer/Concert/ConcertSync/ConcertSyncClient/Source/ConcertSyncClient/Private/ConcertSandboxPlatformFile.cpp
--- a/Engine/Plugins/Developer/Concert/ConcertSync/ConcertSyncClient/Source/ConcertSyncClient/Private/ConcertSandboxPlatformFile.cpp
+++ b/Engine/Plugins/Developer/Concert/ConcertSync/ConcertSyncClient/Source/ConcertSyncClient/Private/ConcertSandboxPlatformFile.cpp
@@ -784,7 +784,7 @@ bool FConcertSandboxPlatformFile::CopyFileWithSCC(
 	}
 	else if (InParam.bShouldMakeWritableIfNoSourceControl)
 	{
-		if (!LowerLevel->SetReadOnly(InToFilename, false))
+		if (LowerLevel->FileExists(InToFilename) && !LowerLevel->SetReadOnly(InToFilename, false))
 		{
 			FText Failure = FText::Format(
 				NSLOCTEXT("Concert", "CanMakeWritable", "Can't make file '{0}' writable when persiting sandbox state!"),
