ConcertSyncClient: force sandbox file changes to be processed immediately when discarding the sandbox

This ensures that the Asset Registry is updated appropriately especially when files or
packages that exist only in the sandbox are deleted during a Multi-User session.

One problematic case was with the use of World Partition levels in a Multi-User session when
actors are added during the session, causing the creation of external packages that exist
only in the session's sandbox. When leaving the session, the sandbox is discarded and the
World Partition level is immediately reloaded. While the sandbox discard was properly
identifying the sandbox-only external packages and deleting them, and it was also registering
the file changes with the FDirectoryWatcherModule, doing so would only queue the file
changes as pending. Without an explicit call to IDirectoryWatcher::Tick(), the file deletes
would be processed *after* the World Partition level had already been reloaded, in which
case the Asset Registry is out of date and still contains a cached description of the added
actors. The actor's external package was indeed deleted from disk though, so the actor
would show as unloaded in the level when the level was reloaded. This could lead to situations
where the load state of the actor in World Partition gets out of sync with the level's list of
loaded actors, and a subsequent background search of the content directory could end up
clearing the stale entries at any time, either of which could potentially lead to crashes.

diff --git a/Engine/Plugins/Developer/Concert/ConcertSync/ConcertSyncClient/Source/ConcertSyncClient/Private/ConcertSandboxPlatformFile.cpp b/Engine/Plugins/Developer/Concert/ConcertSync/ConcertSyncClient/Source/ConcertSyncClient/Private/ConcertSandboxPlatformFile.cpp
--- a/Engine/Plugins/Developer/Concert/ConcertSync/ConcertSyncClient/Source/ConcertSyncClient/Private/ConcertSandboxPlatformFile.cpp
+++ b/Engine/Plugins/Developer/Concert/ConcertSync/ConcertSyncClient/Source/ConcertSyncClient/Private/ConcertSandboxPlatformFile.cpp
@@ -1007,6 +1007,16 @@ void FConcertSandboxPlatformFile::DiscardSandbox(TArray<FName>& OutPackagesPendi
 		if (FileChanges.Num() > 0)
 		{
 			DirectoryWatcherModule->RegisterExternalChanges(FileChanges);
+			if (IDirectoryWatcher* DirectoryWatcher = DirectoryWatcherModule->Get())
+			{
+				// Force the directory watcher to process the file changes
+				// immediately. This ensures that the Asset Registry is
+				// brought up to date especially after file/package deletes,
+				// since we may not otherwise Tick() before packages are hot
+				// reloaded and stale entries could still be present in the
+				// registry.
+				DirectoryWatcher->Tick(0.0f);
+			}
 		}
 	}
 #endif
