ConcertSyncClient: do not mark temporary packages for reload when they are dirtied

Any temporary packages created while in an MU session will not exist on disk when
leaving that session and the sandbox is discarded, so they should not be reloaded.

This resolves an issue where if a world partition level is created and saved while in an
MU session, temporary packages for the external objects will be created and were
previously being marked dirty. When leaving the session, the package manager would
try and fail to reload all of those now defunct temporary packages and would issue an
error dialog about the failure.

diff --git a/Engine/Plugins/Developer/Concert/ConcertSync/ConcertSyncClient/Source/ConcertSyncClient/Private/ConcertClientPackageManager.cpp b/Engine/Plugins/Developer/Concert/ConcertSync/ConcertSyncClient/Source/ConcertSyncClient/Private/ConcertClientPackageManager.cpp
--- a/Engine/Plugins/Developer/Concert/ConcertSync/ConcertSyncClient/Source/ConcertSyncClient/Private/ConcertClientPackageManager.cpp
+++ b/Engine/Plugins/Developer/Concert/ConcertSync/ConcertSyncClient/Source/ConcertSyncClient/Private/ConcertClientPackageManager.cpp
@@ -438,10 +438,12 @@ void FConcertClientPackageManager::HandlePackageDirtyStateChanged(UPackage* InPa
 {
 	check(!InPackage->HasAnyFlags(RF_Transient) || InPackage != GetTransientPackage());
 
-	// Dirty packages are tracked for purge/reload, but 'compiled in' and
-	// 'in memory' cannot be hot purged/reloaded.
+	// Dirty packages are tracked for purge/reload, but 'compiled in',
+	// 'in memory', or temporary packages cannot be hot purged/reloaded.
 	//
-	if (InPackage->IsDirty() && !InPackage->HasAnyPackageFlags(PKG_CompiledIn | PKG_InMemoryOnly))
+	if (InPackage->IsDirty() &&
+			!InPackage->HasAnyPackageFlags(PKG_CompiledIn | PKG_InMemoryOnly) &&
+			!FPackageName::IsTempPackage(InPackage->GetName()))
 	{
 		DirtyPackages.Add(InPackage->GetFName());
 	}
