ConcertSyncClient: do not reload the current level's package if it exists only in the discarded sandbox

When disconnecting from an MU session, destroying the client's package manager
would previously always reload the current level's package if that level used external
objects, since the external packages those objects are in may also have changed.
If the level itself was created while in the MU session though, the level's package
would only exist inside that session's sandbox, so leaving the session would destroy
the sandbox and purge that package, so there is no need to attempt to reload it.

diff --git a/Engine/Plugins/Developer/Concert/ConcertSync/ConcertSyncClient/Source/ConcertSyncClient/Private/ConcertClientPackageManager.cpp b/Engine/Plugins/Developer/Concert/ConcertSync/ConcertSyncClient/Source/ConcertSyncClient/Private/ConcertClientPackageManager.cpp
--- a/Engine/Plugins/Developer/Concert/ConcertSync/ConcertSyncClient/Source/ConcertSyncClient/Private/ConcertClientPackageManager.cpp
+++ b/Engine/Plugins/Developer/Concert/ConcertSync/ConcertSyncClient/Source/ConcertSyncClient/Private/ConcertClientPackageManager.cpp
@@ -119,7 +119,10 @@
 		}
 	}
 
-	// The persistent level should always be reloaded when using external objects.
+	// If the persistent level uses external objects, its package should be
+	// reloaded as long as the package hasn't been marked for purge. If it has
+	// been marked for purge, then it must have only existed in the sandbox
+	// that was just discarded, so it should not be reloaded.
 	UWorld* CurrentWorld = ConcertSyncClientUtil::GetCurrentWorld();
 	if (CurrentWorld)
 	{
@@ -127,7 +130,7 @@
 		if (PersistentLevel && PersistentLevel->IsUsingExternalObjects())
 		{
 			const FName PackageName = PersistentLevel->GetPackage()->GetFName();
-			if (!PackagesPendingHotReload.Contains(PackageName))
+			if (!PackagesPendingPurge.Contains(PackageName) && !PackagesPendingHotReload.Contains(PackageName))
 			{
 				PackagesPendingHotReload.Add(PackageName);
 			}
