ConcertSyncClient: add a utility function for synchronizing the Asset Registry

This moves several private helper functions out of the ConcertSandboxPlatformFileUtils
namespace to instead be public in ConcertSyncClientUtil. It also adds a
SynchronizeAssetRegistry() function that takes the place of the IDirectoryWatcher::Tick()
call that was previously inserted into the sandbox discard process. This new function
can more readily be used elsewhere when file changes need to be applied immediately
in order to bring the Asset Registry up to date.

diff --git a/Engine/Plugins/Developer/Concert/ConcertSync/ConcertSyncClient/Source/ConcertSyncClient/Private/ConcertSandboxPlatformFile.cpp b/Engine/Plugins/Developer/Concert/ConcertSync/ConcertSyncClient/Source/ConcertSyncClient/Private/ConcertSandboxPlatformFile.cpp
--- a/Engine/Plugins/Developer/Concert/ConcertSync/ConcertSyncClient/Source/ConcertSyncClient/Private/ConcertSandboxPlatformFile.cpp
+++ b/Engine/Plugins/Developer/Concert/ConcertSync/ConcertSyncClient/Source/ConcertSyncClient/Private/ConcertSandboxPlatformFile.cpp
@@ -56,43 +56,8 @@
 	return false;
 }
 
-#if WITH_EDITOR
-
-FDirectoryWatcherModule& GetDirectoryWatcherModule()
-{
-	static const FName DirectoryWatcherModuleName = TEXT("DirectoryWatcher");
-	return FModuleManager::LoadModuleChecked<FDirectoryWatcherModule>(DirectoryWatcherModuleName);
 }
 
-FDirectoryWatcherModule* GetDirectoryWatcherModuleIfLoaded()
-{
-	static const FName DirectoryWatcherModuleName = TEXT("DirectoryWatcher");
-	if (FModuleManager::Get().IsModuleLoaded(DirectoryWatcherModuleName))
-	{
-		return &FModuleManager::GetModuleChecked<FDirectoryWatcherModule>(DirectoryWatcherModuleName);
-	}
-	return nullptr;
-}
-
-IDirectoryWatcher* GetDirectoryWatcher()
-{
-	FDirectoryWatcherModule& DirectoryWatcherModule = GetDirectoryWatcherModule();
-	return DirectoryWatcherModule.Get();
-}
-
-IDirectoryWatcher* GetDirectoryWatcherIfLoaded()
-{
-	if (FDirectoryWatcherModule* DirectoryWatcherModule = GetDirectoryWatcherModuleIfLoaded())
-	{
-		return DirectoryWatcherModule->Get();
-	}
-	return nullptr;
-}
-
-#endif
-
-}
-
 FConcertSandboxPlatformFile::FConcertSandboxPlatformFile(const FString& InSandboxRootPath)
 	: SandboxRootPath(InSandboxRootPath)
 	, LowerLevel(nullptr)
@@ -111,7 +76,7 @@
 	FPackageName::OnContentPathDismounted().RemoveAll(this);
 
 #if WITH_EDITOR
-	if (IDirectoryWatcher* DirectoryWatcher = ConcertSandboxPlatformFileUtil::GetDirectoryWatcherIfLoaded())
+	if (IDirectoryWatcher* DirectoryWatcher = ConcertSyncClientUtil::GetDirectoryWatcherIfLoaded())
 	{
 		for (FSandboxMountPoint& SandboxMountPoint : SandboxMountPoints)
 		{
@@ -1002,21 +967,19 @@
 
 #if WITH_EDITOR
 	// Notify that the sandboxed directories have been restored to their original state
-	if (FDirectoryWatcherModule* DirectoryWatcherModule = ConcertSandboxPlatformFileUtil::GetDirectoryWatcherModuleIfLoaded())
+	if (FDirectoryWatcherModule* DirectoryWatcherModule = ConcertSyncClientUtil::GetDirectoryWatcherModuleIfLoaded())
 	{
 		if (FileChanges.Num() > 0)
 		{
 			DirectoryWatcherModule->RegisterExternalChanges(FileChanges);
-			if (IDirectoryWatcher* DirectoryWatcher = DirectoryWatcherModule->Get())
-			{
-				// Force the directory watcher to process the file changes
-				// immediately. This ensures that the Asset Registry is
-				// brought up to date especially after file/package deletes,
-				// since we may not otherwise Tick() before packages are hot
-				// reloaded and stale entries could still be present in the
-				// registry.
-				DirectoryWatcher->Tick(0.0f);
-			}
+
+			// Force a sync of the Asset Registry immediately to process the
+			// file changes. This ensures that the Asset Registry is brought
+			// up to date especially after file/package deletes, since
+			// otherwise it may not get updated before packages are hot
+			// reloaded and stale entries could still be present in the
+			// registry.
+			ConcertSyncClientUtil::SynchronizeAssetRegistry();
 		}
 	}
 #endif
@@ -1111,7 +1074,7 @@
 
 		FSandboxMountPoint& SandboxMountPoint = SandboxMountPoints.Add_GetRef(FSandboxMountPoint{ FConcertSandboxPlatformFilePath(MoveTemp(AbsoluteNonSandboxPath), MoveTemp(AbsoluteSandboxPath)) });
 #if WITH_EDITOR
-		if (IDirectoryWatcher* DirectoryWatcher = ConcertSandboxPlatformFileUtil::GetDirectoryWatcher())
+		if (IDirectoryWatcher* DirectoryWatcher = ConcertSyncClientUtil::GetDirectoryWatcher())
 		{
 			DirectoryWatcher->RegisterDirectoryChangedCallback_Handle(SandboxMountPoint.Path.GetSandboxPath(), IDirectoryWatcher::FDirectoryChanged::CreateRaw(this, &FConcertSandboxPlatformFile::OnDirectoryChanged, SandboxMountPoint.Path), SandboxMountPoint.OnDirectoryChangedHandle, IDirectoryWatcher::IncludeDirectoryChanges);
 		}
@@ -1130,7 +1093,7 @@
 #if WITH_EDITOR
 			if (bShouldRemove && InSandboxMountPoint.OnDirectoryChangedHandle.IsValid())
 			{
-				if (IDirectoryWatcher* DirectoryWatcher = ConcertSandboxPlatformFileUtil::GetDirectoryWatcherIfLoaded())
+				if (IDirectoryWatcher* DirectoryWatcher = ConcertSyncClientUtil::GetDirectoryWatcherIfLoaded())
 				{
 					DirectoryWatcher->UnregisterDirectoryChangedCallback_Handle(InSandboxMountPoint.Path.GetSandboxPath(), InSandboxMountPoint.OnDirectoryChangedHandle);
 					InSandboxMountPoint.OnDirectoryChangedHandle.Reset();
@@ -1231,7 +1194,7 @@
 	}
 
 #if WITH_EDITOR
-	if (FDirectoryWatcherModule* DirectoryWatcherModule = ConcertSandboxPlatformFileUtil::GetDirectoryWatcherModuleIfLoaded())
+	if (FDirectoryWatcherModule* DirectoryWatcherModule = ConcertSyncClientUtil::GetDirectoryWatcherModuleIfLoaded())
 	{
 		FFileChangeData FileChange(InPath.GetNonSandboxPath(), FFileChangeData::FCA_Removed);
 		DirectoryWatcherModule->RegisterExternalChanges(TArrayView<const FFileChangeData>(&FileChange, 1));
@@ -1345,7 +1308,7 @@
 		return;
 	}
 
-	if (FDirectoryWatcherModule* DirectoryWatcherModule = ConcertSandboxPlatformFileUtil::GetDirectoryWatcherModuleIfLoaded())
+	if (FDirectoryWatcherModule* DirectoryWatcherModule = ConcertSyncClientUtil::GetDirectoryWatcherModuleIfLoaded())
 	{
 		TArray<FFileChangeData> RemappedFileChanges;
 		RemappedFileChanges.Reserve(FileChanges.Num());
diff --git a/Engine/Plugins/Developer/Concert/ConcertSync/ConcertSyncClient/Source/ConcertSyncClient/Private/ConcertSyncClientUtil.cpp b/Engine/Plugins/Developer/Concert/ConcertSync/ConcertSyncClient/Source/ConcertSyncClient/Private/ConcertSyncClientUtil.cpp
--- a/Engine/Plugins/Developer/Concert/ConcertSync/ConcertSyncClient/Source/ConcertSyncClient/Private/ConcertSyncClientUtil.cpp
+++ b/Engine/Plugins/Developer/Concert/ConcertSync/ConcertSyncClient/Source/ConcertSyncClient/Private/ConcertSyncClientUtil.cpp
@@ -30,15 +30,18 @@
 #include "Framework/Application/SlateApplication.h"
 
 #if WITH_EDITOR
+	#include "DirectoryWatcherModule.h"
 	#include "Editor.h"
-	#include "UnrealEdGlobals.h"
 	#include "Editor/UnrealEdEngine.h"
+	#include "FileHelpers.h"
+	#include "GameMapsSettings.h"
+	#include "IDirectoryWatcher.h"
+	#include "Modules/ModuleManager.h"
+	#include "ObjectTools.h"	
 	#include "PackageTools.h"
-	#include "ObjectTools.h"
 	#include "Selection.h"
 	#include "Subsystems/AssetEditorSubsystem.h"
-	#include "GameMapsSettings.h"
-	#include "FileHelpers.h"
+	#include "UnrealEdGlobals.h"
 #endif
 
 namespace ConcertSyncClientUtil
@@ -379,6 +382,54 @@
 	}
 }
 
+#if WITH_EDITOR
+
+FDirectoryWatcherModule& GetDirectoryWatcherModule()
+{
+	static const FName DirectoryWatcherModuleName = TEXT("DirectoryWatcher");
+	return FModuleManager::LoadModuleChecked<FDirectoryWatcherModule>(DirectoryWatcherModuleName);
+}
+
+FDirectoryWatcherModule* GetDirectoryWatcherModuleIfLoaded()
+{
+	static const FName DirectoryWatcherModuleName = TEXT("DirectoryWatcher");
+	if (FModuleManager::Get().IsModuleLoaded(DirectoryWatcherModuleName))
+	{
+		return &FModuleManager::GetModuleChecked<FDirectoryWatcherModule>(DirectoryWatcherModuleName);
+	}
+	return nullptr;
+}
+
+IDirectoryWatcher* GetDirectoryWatcher()
+{
+	FDirectoryWatcherModule& DirectoryWatcherModule = GetDirectoryWatcherModule();
+	return DirectoryWatcherModule.Get();
+}
+
+IDirectoryWatcher* GetDirectoryWatcherIfLoaded()
+{
+	if (FDirectoryWatcherModule* DirectoryWatcherModule = GetDirectoryWatcherModuleIfLoaded())
+	{
+		return DirectoryWatcherModule->Get();
+	}
+	return nullptr;
+}
+
+#endif // WITH_EDITOR
+
+void SynchronizeAssetRegistry()
+{
+#if WITH_EDITOR
+	IDirectoryWatcher* DirectoryWatcher = GetDirectoryWatcherIfLoaded();
+	if (!DirectoryWatcher)
+	{
+		return;
+	}
+
+	DirectoryWatcher->Tick(0.0f);
+#endif // WITH_EDITOR
+}
+
 void HotReloadPackages(TArrayView<const FName> InPackageNames)
 {
 	if (InPackageNames.Num() == 0)
diff --git a/Engine/Plugins/Developer/Concert/ConcertSync/ConcertSyncClient/Source/ConcertSyncClient/Private/ConcertSyncClientUtil.h b/Engine/Plugins/Developer/Concert/ConcertSync/ConcertSyncClient/Source/ConcertSyncClient/Private/ConcertSyncClientUtil.h
--- a/Engine/Plugins/Developer/Concert/ConcertSync/ConcertSyncClient/Source/ConcertSyncClient/Private/ConcertSyncClientUtil.h
+++ b/Engine/Plugins/Developer/Concert/ConcertSync/ConcertSyncClient/Source/ConcertSyncClient/Private/ConcertSyncClientUtil.h
@@ -6,6 +6,11 @@
 #include "Misc/EnumClassFlags.h"
 #include "Containers/ArrayView.h"
 
+#if WITH_EDITOR
+	#include "DirectoryWatcherModule.h"
+	#include "IDirectoryWatcher.h"
+#endif
+
 class UObject;
 class UPackage;
 class UStruct;
@@ -90,6 +95,43 @@
 
 	void FlushPackageLoading(const FString& InPackageName, bool bForceBulkDataLoad = true);
 
+#if WITH_EDITOR
+	/** Gets the DirectoryWatcher module.
+	 *
+	 *  The module will be loaded if it is not currently loaded.
+	 * 
+	 *  @return: the DirectoryWatcher module
+	 */
+	FDirectoryWatcherModule& GetDirectoryWatcherModule();
+
+	/** Gets the DirectoryWatcher module if it is currently loaded.
+	 *
+	 *  @return: a pointer to the DirectoryWatcher module if is currently loaded, or nullptr otherwise
+	 */
+	FDirectoryWatcherModule* GetDirectoryWatcherModuleIfLoaded();
+
+	/** Gets the directory watcher.
+	 *
+	 *  The DirectoryWatcher module will be loaded if it is not currently loaded.
+	 *
+	 *  @return: the directory watcher if the platform supports directory watching, or nullptr otherwise
+	 */
+	IDirectoryWatcher* GetDirectoryWatcher();
+
+	/** Gets the directory watcher if the DirectoryWatcher module is currently loaded.
+	 *
+	 *  @return: the directory watcher if the DirectoryWatcher module is currently loaded and the platform supports directory watching, or nullptr otherwise
+	 */
+	IDirectoryWatcher* GetDirectoryWatcherIfLoaded();
+#endif // WITH_EDITOR
+
+	/** Synchronizes the Asset Registry.
+	 * 
+	 *  This ensures that any pending file changes are completed and that the
+	 *  Asset Registry is updated to correctly reflect those changes.
+	 */
+	void SynchronizeAssetRegistry();
+
 	void HotReloadPackages(TArrayView<const FName> InPackageNames);
 
 	void PurgePackages(TArrayView<const FName> InPackageNames);
