Add a CVar to delay timeout so that the transactions do not pile up if the user
is not interacting with the session.

diff --git a/Engine/Plugins/Developer/Concert/ConcertSync/ConcertSyncClient/Source/ConcertSyncClient/Private/ConcertSyncClientUtil.cpp b/Engine/Plugins/Developer/Concert/ConcertSync/ConcertSyncClient/Source/ConcertSyncClient/Private/ConcertSyncClientUtil.cpp
--- a/Engine/Plugins/Developer/Concert/ConcertSync/ConcertSyncClient/Source/ConcertSyncClient/Private/ConcertSyncClientUtil.cpp
+++ b/Engine/Plugins/Developer/Concert/ConcertSync/ConcertSyncClient/Source/ConcertSyncClient/Private/ConcertSyncClientUtil.cpp
@@ -51,6 +51,10 @@
 	TEXT("Concert.DelayTransactionsWhileEditing"), 0,
 	TEXT("Focus is lost by the editor when a transaction is applied. This variable suspends applying a transaction until the user has removed focus on editable UI."));
 
+static TAutoConsoleVariable<int32> CVarDelayApplyingTransactionsWaitTimeout(
+	TEXT("Concert.DelayTransactionsWhileEditingTimeout"), 5,
+	TEXT("When Concert.DelayTransactionsWhileEditing is enabled we make sure the user has not been idle too long to prevent transactions from stacking up. The timeout value is specified in seconds."));
+
 bool UserIsEditing()
 {
 #if WITH_EDITOR
@@ -60,12 +64,22 @@
 		static FName SMultiLineEditableTextType(TEXT("SMultiLineEditableText"));
 
 		bool bEditable = false;
-		FSlateApplication::Get().ForEachUser([&bEditable](FSlateUser& User) {
+		FSlateApplication& App = FSlateApplication::Get();
+		App.ForEachUser([&bEditable](FSlateUser& User) {
 			TSharedPtr<SWidget> FocusedWidget = User.GetFocusedWidget();
 
 			bool bCanEdit = FocusedWidget && (FocusedWidget->GetType() == SEditableTextType || FocusedWidget->GetType() == SMultiLineEditableTextType);
 			bEditable |= bCanEdit;
 		});
+		if (bEditable)
+		{
+			double LastUpdateTime = App.GetLastUserInteractionTime();
+			double Duration = App.GetCurrentTime() - LastUpdateTime;
+			if (static_cast<int32>(Duration) > CVarDelayApplyingTransactionsWaitTimeout.GetValueOnAnyThread())
+			{
+				return false;
+			}
+		}
 		return bEditable;
 	}
 #endif
