Keep remotely deleted actors alive during GC passes to allow editor users to save these changes locally, since these replicated Concert changes aren't transacted.

#rnx

diff --git a/Engine/Plugins/Developer/Concert/ConcertSync/ConcertSyncClient/Source/ConcertSyncClient/Private/ConcertClientTransactionBridge.cpp b/Engine/Plugins/Developer/Concert/ConcertSync/ConcertSyncClient/Source/ConcertSyncClient/Private/ConcertClientTransactionBridge.cpp
--- a/Engine/Plugins/Developer/Concert/ConcertSync/ConcertSyncClient/Source/ConcertSyncClient/Private/ConcertClientTransactionBridge.cpp
+++ b/Engine/Plugins/Developer/Concert/ConcertSync/ConcertSyncClient/Source/ConcertSyncClient/Private/ConcertClientTransactionBridge.cpp
@@ -177,7 +177,7 @@ struct FEditorTransactionNotification
 };
 #endif
 
-void ProcessTransactionEvent(const FConcertTransactionEventBase& InEvent, const FConcertSessionVersionInfo* InVersionInfo, const TArray<FName>& InPackagesToProcess, const FConcertLocalIdentifierTable* InLocalIdentifierTablePtr, const bool bIsSnapshot, const FConcertSyncWorldRemapper& WorldRemapper)
+void ProcessTransactionEvent(const FConcertTransactionEventBase& InEvent, const FConcertSessionVersionInfo* InVersionInfo, const TArray<FName>& InPackagesToProcess, const FConcertLocalIdentifierTable* InLocalIdentifierTablePtr, const bool bIsSnapshot, const FConcertSyncWorldRemapper& WorldRemapper, TSet<AActor*>* DeletedActors)
 {
 	// Transactions are applied in multiple-phases...
 	//	1) Find or create all objects in the transaction (to handle object-interdependencies in the serialized data)
@@ -357,10 +357,19 @@ void ProcessTransactionEvent(const FConcertTransactionEventBase& InEvent, const
 		}
 
 		// Finish spawning any newly created actors
-		if (TransactionObjectRef.NeedsPostSpawn())
+		if (AActor* TransactionActor = Cast<AActor>(TransactionObject))
 		{
-			AActor* TransactionActor = CastChecked<AActor>(TransactionObject);
-			TransactionActor->FinishSpawning(FTransform(), true);
+			if (ObjectUpdate.ObjectData.bIsPendingKill)
+			{
+				if (DeletedActors)
+				{
+					DeletedActors->Add(TransactionActor);
+				}
+			}
+			else if (TransactionObjectRef.NeedsPostSpawn())
+			{
+				TransactionActor->FinishSpawning(FTransform(), true);
+			}
 		}
 
 #if WITH_EDITOR
@@ -482,12 +491,17 @@ FOnApplyTransaction& FConcertClientTransactionBridge::OnApplyTransaction()
 
 void FConcertClientTransactionBridge::ApplyRemoteTransaction(const FConcertTransactionEventBase& InEvent, const FConcertSessionVersionInfo* InVersionInfo, const TArray<FName>& InPackagesToProcess, const FConcertLocalIdentifierTable* InLocalIdentifierTablePtr, const bool bIsSnapshot)
 {
-	ConcertClientTransactionBridgeUtil::ProcessTransactionEvent(InEvent, InVersionInfo, InPackagesToProcess, InLocalIdentifierTablePtr, bIsSnapshot, FConcertSyncWorldRemapper());
+	ConcertClientTransactionBridgeUtil::ProcessTransactionEvent(InEvent, InVersionInfo, InPackagesToProcess, InLocalIdentifierTablePtr, bIsSnapshot, FConcertSyncWorldRemapper(), nullptr);
 }
 
 void FConcertClientTransactionBridge::ApplyRemoteTransaction(const FConcertTransactionEventBase& InEvent, const FConcertSessionVersionInfo* InVersionInfo, const TArray<FName>& InPackagesToProcess, const FConcertLocalIdentifierTable* InLocalIdentifierTablePtr, const bool bIsSnapshot, const FConcertSyncWorldRemapper& ConcertSyncWorldRemapper)
 {
-	ConcertClientTransactionBridgeUtil::ProcessTransactionEvent(InEvent, InVersionInfo, InPackagesToProcess, InLocalIdentifierTablePtr, bIsSnapshot, ConcertSyncWorldRemapper);
+	ConcertClientTransactionBridgeUtil::ProcessTransactionEvent(InEvent, InVersionInfo, InPackagesToProcess, InLocalIdentifierTablePtr, bIsSnapshot, ConcertSyncWorldRemapper, nullptr);
+}
+
+void FConcertClientTransactionBridge::ApplyRemoteTransaction(const FConcertTransactionEventBase& InEvent, const FConcertSessionVersionInfo* InVersionInfo, const TArray<FName>& InPackagesToProcess, const FConcertLocalIdentifierTable* InLocalIdentifierTablePtr, const bool bIsSnapshot, const FConcertSyncWorldRemapper& ConcertSyncWorldRemapper, TSet<AActor*>* DeletedActors)
+{
+	ConcertClientTransactionBridgeUtil::ProcessTransactionEvent(InEvent, InVersionInfo, InPackagesToProcess, InLocalIdentifierTablePtr, bIsSnapshot, ConcertSyncWorldRemapper, DeletedActors);
 }
 
 bool& FConcertClientTransactionBridge::GetIgnoreLocalTransactionsRef()
diff --git a/Engine/Plugins/Developer/Concert/ConcertSync/ConcertSyncClient/Source/ConcertSyncClient/Private/ConcertClientTransactionBridge.h b/Engine/Plugins/Developer/Concert/ConcertSync/ConcertSyncClient/Source/ConcertSyncClient/Private/ConcertClientTransactionBridge.h
--- a/Engine/Plugins/Developer/Concert/ConcertSync/ConcertSyncClient/Source/ConcertSyncClient/Private/ConcertClientTransactionBridge.h
+++ b/Engine/Plugins/Developer/Concert/ConcertSync/ConcertSyncClient/Source/ConcertSyncClient/Private/ConcertClientTransactionBridge.h
@@ -21,6 +21,7 @@ public:
 
 	virtual void ApplyRemoteTransaction(const FConcertTransactionEventBase& InEvent, const FConcertSessionVersionInfo* InVersionInfo, const TArray<FName>& InPackagesToProcess, const FConcertLocalIdentifierTable* InLocalIdentifierTablePtr, const bool bIsSnapshot) override;
 	virtual void ApplyRemoteTransaction(const FConcertTransactionEventBase& InEvent, const FConcertSessionVersionInfo* InVersionInfo, const TArray<FName>& InPackagesToProcess, const FConcertLocalIdentifierTable* InLocalIdentifierTablePtr, const bool bIsSnapshot, const FConcertSyncWorldRemapper& ConcertSyncWorldRemapper) override;
+	virtual void ApplyRemoteTransaction(const FConcertTransactionEventBase& InEvent, const FConcertSessionVersionInfo* InVersionInfo, const TArray<FName>& InPackagesToProcess, const FConcertLocalIdentifierTable* InLocalIdentifierTablePtr, const bool bIsSnapshot, const class FConcertSyncWorldRemapper& ConcertSyncWorldRemapper, TSet<AActor*>* DeletedActors) override;
 	virtual bool& GetIgnoreLocalTransactionsRef() override;
 
 	virtual void RegisterTransactionFilter(FName FilterName, FTransactionFilterDelegate FilterHandle) override;
diff --git a/Engine/Plugins/Developer/Concert/ConcertSync/ConcertSyncClient/Source/ConcertSyncClient/Public/IConcertClientTransactionBridge.h b/Engine/Plugins/Developer/Concert/ConcertSync/ConcertSyncClient/Source/ConcertSyncClient/Public/IConcertClientTransactionBridge.h
--- a/Engine/Plugins/Developer/Concert/ConcertSync/ConcertSyncClient/Source/ConcertSyncClient/Public/IConcertClientTransactionBridge.h
+++ b/Engine/Plugins/Developer/Concert/ConcertSync/ConcertSyncClient/Source/ConcertSyncClient/Public/IConcertClientTransactionBridge.h
@@ -127,6 +127,18 @@ public:
 	 */
 	virtual void ApplyRemoteTransaction(const FConcertTransactionEventBase& InEvent, const FConcertSessionVersionInfo* InVersionInfo, const TArray<FName>& InPackagesToProcess, const FConcertLocalIdentifierTable* InLocalIdentifierTablePtr, const bool bIsSnapshot, const class FConcertSyncWorldRemapper& ConcertSyncWorldRemapper) = 0;
 
+	/**
+	 * Apply a remote transaction event to this local instance.
+	 * @param InEvent					The event to apply.
+	 * @param InVersionInfo				The version information for the serialized data in the event, or null if the event should be serialized using the compiled in version info.
+	 * @param InPackagesToProcess		The list of packages to apply changes for, or an empty array to apply all changes.
+	 * @param InLocalIdentifierTablePtr The local identifier table for the event data (if any).
+	 * @param bIsSnapshot				True if this transaction event was a snapshot rather than a finalized transaction.
+	 * @param ConcertSyncWorldRemapper	Remapper to use in the case the current world is different from the one sent in the transation.
+	 * @param DeletedActors				Optional set of actor pointers to gather deleted actors.
+	 */
+	virtual void ApplyRemoteTransaction(const FConcertTransactionEventBase& InEvent, const FConcertSessionVersionInfo* InVersionInfo, const TArray<FName>& InPackagesToProcess, const FConcertLocalIdentifierTable* InLocalIdentifierTablePtr, const bool bIsSnapshot, const class FConcertSyncWorldRemapper& ConcertSyncWorldRemapper, TSet<AActor*>* DeletedActors) = 0;
+
 	/** Callback to register delegate for handling transaction events */
 	virtual void RegisterTransactionFilter(FName FilterName, FTransactionFilterDelegate FilterHandle) = 0;
 
