ConcertApp, ConcertSync: add an option/cvar for closing Sequencer when a sequence is closed by a remote user

diff --git a/Engine/Plugins/Developer/Concert/ConcertApp/MultiUserClient/Source/MultiUserClient/Private/ConcertWorkspaceUI.cpp b/Engine/Plugins/Developer/Concert/ConcertApp/MultiUserClient/Source/MultiUserClient/Private/ConcertWorkspaceUI.cpp
--- a/Engine/Plugins/Developer/Concert/ConcertApp/MultiUserClient/Source/MultiUserClient/Private/ConcertWorkspaceUI.cpp
+++ b/Engine/Plugins/Developer/Concert/ConcertApp/MultiUserClient/Source/MultiUserClient/Private/ConcertWorkspaceUI.cpp
@@ -448,6 +448,32 @@ public:
 				FSlateIcon(FConcertFrontendStyle::GetStyleSetName(), "Concert.Sequencer.SyncSequence", "Concert.Sequencer.SyncSequence.Small"),
 				EUserInterfaceActionType::ToggleButton
 			);
+
+			// Toggle remote close
+			ToolbarBuilder.AddToolBarButton(
+				FUIAction(
+					FExecuteAction::CreateLambda([]()
+					{
+						TSharedPtr<IConcertSyncClient> SyncClientPin = IConcertSyncClientModule::Get().GetClient(TEXT("MultiUser"));
+						if (SyncClientPin && SyncClientPin->GetSequencerManager())
+						{
+							IConcertClientSequencerManager* SequencerManager = SyncClientPin->GetSequencerManager();
+							SequencerManager->SetSequencerRemoteClose(!SequencerManager->IsSequencerRemoteCloseEnabled());
+						}
+					}),
+					FCanExecuteAction(),
+					FIsActionChecked::CreateLambda([]()
+					{
+						TSharedPtr<IConcertSyncClient> SyncClientPin = IConcertSyncClientModule::Get().GetClient(TEXT("MultiUser"));
+						return SyncClientPin.IsValid() && SyncClientPin->GetSequencerManager() && SyncClientPin->GetSequencerManager()->IsSequencerRemoteCloseEnabled();
+					})
+				),
+				NAME_None,
+				LOCTEXT("ToggleRemoteCloseLabel", "Remote Close"),
+				LOCTEXT("ToggleRemoteCloseTooltip", "Toggle Multi-User Remote Close. If the option is enabled, any open Sequencers for a sequence will be closed when another user in the Multi-User session closes that sequence."),
+				FSlateIcon(FConcertFrontendStyle::GetStyleSetName(), "Concert.Sequencer.SyncSequence", "Concert.Sequencer.SyncSequence.Small"),
+				EUserInterfaceActionType::ToggleButton
+			);
 		}
 		ToolbarBuilder.EndSection();
 
diff --git a/Engine/Plugins/Developer/Concert/ConcertSync/ConcertSyncClient/Source/ConcertSyncClient/Private/ConcertClientSequencerManager.cpp b/Engine/Plugins/Developer/Concert/ConcertSync/ConcertSyncClient/Source/ConcertSyncClient/Private/ConcertClientSequencerManager.cpp
--- a/Engine/Plugins/Developer/Concert/ConcertSync/ConcertSyncClient/Source/ConcertSyncClient/Private/ConcertClientSequencerManager.cpp
+++ b/Engine/Plugins/Developer/Concert/ConcertSync/ConcertSyncClient/Source/ConcertSyncClient/Private/ConcertClientSequencerManager.cpp
@@ -42,10 +42,13 @@ static TAutoConsoleVariable<int32> CVarEnablePlaybackSync(TEXT("Concert.EnableSe
 // Enable Sequence Playing on game client
 static TAutoConsoleVariable<int32> CVarEnableSequencePlayer(TEXT("Concert.EnableSequencePlayer"), 1, TEXT("Enable Concert Sequence Players on `-game` client."));
 
-// Enable opening Sequencer on remote machine whenever a sequencer is opened, if both instance have this option on.
+// Enable opening Sequencer on remote machine whenever a Sequencer is opened, if both instances have this option on.
 static TAutoConsoleVariable<int32> CVarEnableRemoteSequencerOpen(TEXT("Concert.EnableOpenRemoteSequencer"), 1, TEXT("Enable Concert remote Sequencer opening."));
 
-// Enable opening Sequencer on remote machine whenever a sequencer is opened, if both instance have this option on.
+// Enable closing Sequencer for this user when a remote user closes the sequence.
+static TAutoConsoleVariable<int32> CVarEnableRemoteSequencerClose(TEXT("Concert.EnableCloseRemoteSequencer"), 0, TEXT("Enable Concert remote Sequencer closing."));
+
+// Enable synchronizing the timeline of unrelated Sequencers on remote machine whenever a Sequencer state event is received, if both instances have this option on.
 static TAutoConsoleVariable<int32> CVarEnableUnrelatedTimelineSync(TEXT("Concert.EnableUnrelatedTimelineSync"), 0, TEXT("Enable syncing unrelated sequencer timeline."));
 
 // Enable always closing player on remote machine whenever a sequencer is closed on an editor.
@@ -232,11 +235,21 @@ bool FConcertClientSequencerManager::IsSequencerRemoteOpenEnabled() const
 	return CVarEnableRemoteSequencerOpen.GetValueOnAnyThread() > 0;
 }
 
+bool FConcertClientSequencerManager::IsSequencerRemoteCloseEnabled() const
+{
+	return CVarEnableRemoteSequencerClose.GetValueOnAnyThread() > 0;
+}
+
 void FConcertClientSequencerManager::SetSequencerRemoteOpen(bool bEnable)
 {
 	CVarEnableRemoteSequencerOpen->AsVariable()->Set(bEnable ? 1 : 0);
 }
 
+void FConcertClientSequencerManager::SetSequencerRemoteClose(bool bEnable)
+{
+	CVarEnableRemoteSequencerClose->AsVariable()->Set(bEnable ? 1 : 0);
+}
+
 bool FConcertClientSequencerManager::ShouldAlwaysCloseGameSequencerPlayer() const
 {
 	return CVarAlwaysCloseGamePlayerOnCloseEvent.GetValueOnAnyThread() > 0;
@@ -474,17 +487,37 @@ void FConcertClientSequencerManager::ApplyCloseEvent(const FConcertSequencerClos
 	FConcertSequencerState* SequencerState = SequencerStates.Find(*CloseEvent.SequenceObjectPath);
 	if (SequencerState)
 	{
+		const TArray<FOpenSequencerData*, TInlineAllocator<1>> OpenSequencersForObject = GatherRootSequencersByState(SequencerState->SequenceObjectPath);
+
 		// if the event was that a sequencer that was in controller playback mode was closed, stop playback
 		if (CloseEvent.bControllerClose)
 		{
 			SequencerState->PlayerStatus = EConcertMovieScenePlayerStatus::Stopped;
-			for (FOpenSequencerData* OpenSequencer : GatherRootSequencersByState(SequencerState->SequenceObjectPath))
+			for (FOpenSequencerData* OpenSequencer : OpenSequencersForObject)
 			{
 				OpenSequencer->PlaybackMode = EPlaybackMode::Undefined;
 				TSharedPtr<ISequencer> Sequencer = OpenSequencer->WeakSequencer.Pin();
 				Sequencer->SetPlaybackStatus(EMovieScenePlayerStatus::Stopped);
 			}
 		}
+
+		if (GEditor && IsSequencerRemoteCloseEnabled())
+		{
+			for (FOpenSequencerData* OpenSequencer : OpenSequencersForObject)
+			{
+				TSharedPtr<ISequencer> Sequencer = OpenSequencer->WeakSequencer.Pin();
+				UMovieSceneSequence* Sequence = Sequencer.IsValid() ? Sequencer->GetRootMovieSceneSequence() : nullptr;
+
+				// Verify that this Sequencer is open for the particular
+				// sequence being closed since the array of open Sequencers
+				// may contain unrelated Sequencers.
+				if (Sequence && (Sequence->GetPathName() == CloseEvent.SequenceObjectPath))
+				{
+					GEditor->GetEditorSubsystem<UAssetEditorSubsystem>()->CloseAllEditorsForAsset(Sequence);
+				}
+			}
+		}
+
 		// Discard the state if it's no longer opened by anyone.
 		//
 		if (CloseEvent.EditorsWithSequencerOpened == 0)
diff --git a/Engine/Plugins/Developer/Concert/ConcertSync/ConcertSyncClient/Source/ConcertSyncClient/Private/ConcertClientSequencerManager.h b/Engine/Plugins/Developer/Concert/ConcertSync/ConcertSyncClient/Source/ConcertSyncClient/Private/ConcertClientSequencerManager.h
--- a/Engine/Plugins/Developer/Concert/ConcertSync/ConcertSyncClient/Source/ConcertSyncClient/Private/ConcertClientSequencerManager.h
+++ b/Engine/Plugins/Developer/Concert/ConcertSync/ConcertSyncClient/Source/ConcertSyncClient/Private/ConcertClientSequencerManager.h
@@ -86,6 +86,11 @@ public:
 	 */
 	virtual bool IsSequencerRemoteOpenEnabled() const override;
 
+	/**
+	 * @return true if the remote close option is enabled.
+	 */
+	virtual bool IsSequencerRemoteCloseEnabled() const override;
+
 	/**
 	 * Checks the CVar to see if we are allowed to forcefully close the player on game instances.
 	 *
@@ -95,12 +100,20 @@ public:
 
 	/**
 	 * Set the remote open option in Multi-User
-	 * which open sequencer for other user when this option is enabled on both user machines.
-	 * 
+	 * which opens Sequencer for other users when this option is enabled on both user machines.
+	 *
 	 * @param bEnable The value to set for the remote open option
 	 */
 	virtual void SetSequencerRemoteOpen(bool bEnable) override;
 
+	/**
+	 * Set the remote close option in Multi-User
+	 * which closes Sequencer for this user when the sequence is closed by a remote user.
+	 *
+	 * @param bEnable The value to set for the remote close option
+	 */
+	virtual void SetSequencerRemoteClose(bool bEnable) override;
+
 	/** Assign the current active workspace to this sequencer. */
 	void SetActiveWorkspace(TSharedPtr<FConcertClientWorkspace> Workspace);
 
diff --git a/Engine/Plugins/Developer/Concert/ConcertSync/ConcertSyncClient/Source/ConcertSyncClient/Public/IConcertClientSequencerManager.h b/Engine/Plugins/Developer/Concert/ConcertSync/ConcertSyncClient/Source/ConcertSyncClient/Public/IConcertClientSequencerManager.h
--- a/Engine/Plugins/Developer/Concert/ConcertSync/ConcertSyncClient/Source/ConcertSyncClient/Public/IConcertClientSequencerManager.h
+++ b/Engine/Plugins/Developer/Concert/ConcertSync/ConcertSyncClient/Source/ConcertSyncClient/Public/IConcertClientSequencerManager.h
@@ -40,11 +40,24 @@ public:
 	 */
 	virtual bool IsSequencerRemoteOpenEnabled() const = 0;
 
+	/**
+	 * @return true if the remote close option is enabled.
+	 */
+	virtual bool IsSequencerRemoteCloseEnabled() const = 0;
+
 	/**
 	 * Set the remote open option in Multi-User
-	 * which open sequencer for other user when this option is enabled on both user machines.
+	 * which opens Sequencer for other users when this option is enabled on both user machines.
 	 * 
 	 * @param bEnable The value to set for the remote open option
 	 */
 	virtual void SetSequencerRemoteOpen(bool bEnable) = 0;
+
+	/**
+	 * Set the remote close option in Multi-User
+	 * which closes Sequencer for this user when the sequence is closed by a remote user.
+	 *
+	 * @param bEnable The value to set for the remote close option
+	 */
+	virtual void SetSequencerRemoteClose(bool bEnable) = 0;
 };
