Add flags to EditorEngine to control when we note selection changes and broadcast Undo/Redo events.

With Multi-user, applied transactions shouldn't take focus away from the user editing a field.

diff --git a/Engine/Plugins/Developer/Concert/ConcertSync/ConcertSyncClient/Source/ConcertSyncClient/Private/ConcertClientTransactionBridge.cpp b/Engine/Plugins/Developer/Concert/ConcertSync/ConcertSyncClient/Source/ConcertSyncClient/Private/ConcertClientTransactionBridge.cpp
--- a/Engine/Plugins/Developer/Concert/ConcertSync/ConcertSyncClient/Source/ConcertSyncClient/Private/ConcertClientTransactionBridge.cpp
+++ b/Engine/Plugins/Developer/Concert/ConcertSync/ConcertSyncClient/Source/ConcertSyncClient/Private/ConcertClientTransactionBridge.cpp
@@ -16,8 +16,10 @@
 	#include "UnrealEdGlobals.h"
 	#include "Editor/UnrealEdEngine.h"
 	#include "Editor/TransBuffer.h"
+	#include "Engine/Selection.h"
 #endif
 
+
 #define LOCTEXT_NAMESPACE "ConcertClientTransactionBridge"
 
 namespace ConcertClientTransactionBridgeUtil
@@ -113,7 +115,27 @@ struct FEditorTransactionNotification
 		, TransBuffer(GUnrealEd ? Cast<UTransBuffer>(GUnrealEd->Trans) : nullptr)
 		, bOrigSquelchTransactionNotification(GEditor && GEditor->bSquelchTransactionNotification)
 		, bOrigNotifyUndoRedoSelectionChange(GEditor && GEditor->bNotifyUndoRedoSelectionChange)
+		, bOrigIgnoreSelectionChange(GEditor && GEditor->bIgnoreSelectionChange)
+		, bOrigSuspendBroadcastPostUndoRedo(GEditor && GEditor->bSuspendBroadcastPostUndoRedo)
+	{
+	}
+
+	bool IsTransactionBufferObjectsInSelection() const
 	{
+		if (!TransBuffer || !GEditor)
+		{
+			return false;
+		}
+
+		USelection* SelectedActors = GEditor->GetSelectedActors();
+		for (FSelectionIterator Iter(*SelectedActors); Iter; ++Iter)
+		{
+			if (TransBuffer->IsObjectInTransactionBuffer(*Iter))
+			{
+				return true;
+			}
+		}
+		return false;
 	}
 
 	void PreUndo()
@@ -124,6 +146,12 @@ struct FEditorTransactionNotification
 			GEditor->bNotifyUndoRedoSelectionChange = true;
 			if (TransBuffer)
 			{
+				if (ConcertSyncClientUtil::IsUserEditing() && !IsTransactionBufferObjectsInSelection())
+				{
+					GEditor->bIgnoreSelectionChange = true;
+					GEditor->bSuspendBroadcastPostUndoRedo = true;
+				}
+
 				TransBuffer->OnBeforeRedoUndo().Broadcast(TransactionContext);
 			}
 		}
@@ -139,6 +167,8 @@ struct FEditorTransactionNotification
 			}
 			GEditor->bSquelchTransactionNotification = bOrigSquelchTransactionNotification;
 			GEditor->bNotifyUndoRedoSelectionChange = bOrigNotifyUndoRedoSelectionChange;
+			GEditor->bIgnoreSelectionChange = bOrigIgnoreSelectionChange;
+			GEditor->bSuspendBroadcastPostUndoRedo = bOrigSuspendBroadcastPostUndoRedo;
 		}
 	}
 
@@ -174,6 +204,8 @@ struct FEditorTransactionNotification
 	UTransBuffer* TransBuffer;
 	bool bOrigSquelchTransactionNotification;
 	bool bOrigNotifyUndoRedoSelectionChange;
+	bool bOrigIgnoreSelectionChange;
+	bool bOrigSuspendBroadcastPostUndoRedo;
 };
 #endif
 
diff --git a/Engine/Plugins/Developer/Concert/ConcertSync/ConcertSyncClient/Source/ConcertSyncClient/Private/ConcertClientWorkspace.cpp b/Engine/Plugins/Developer/Concert/ConcertSync/ConcertSyncClient/Source/ConcertSyncClient/Private/ConcertClientWorkspace.cpp
--- a/Engine/Plugins/Developer/Concert/ConcertSync/ConcertSyncClient/Source/ConcertSyncClient/Private/ConcertClientWorkspace.cpp
+++ b/Engine/Plugins/Developer/Concert/ConcertSync/ConcertSyncClient/Source/ConcertSyncClient/Private/ConcertClientWorkspace.cpp
@@ -862,7 +862,7 @@ void FConcertClientWorkspace::OnEndFrame()
 		LiveSession->GetSession().SendCustomEvent(FConcertWorkspaceSyncAndFinalizeCompletedEvent(), LiveSession->GetSession().GetSessionServerEndpointId(), EConcertMessageFlags::ReliableOrdered);
 	}
 
-	if (bHasSyncedWorkspace && CanProcessPendingPackages() && !ConcertSyncClientUtil::UserIsEditing())
+	if (bHasSyncedWorkspace && CanProcessPendingPackages() && !ConcertSyncClientUtil::ShouldDelayTransaction())
 	{
 		if (PackageManager)
 		{
diff --git a/Engine/Plugins/Developer/Concert/ConcertSync/ConcertSyncClient/Source/ConcertSyncClient/Private/ConcertSyncClientUtil.cpp b/Engine/Plugins/Developer/Concert/ConcertSync/ConcertSyncClient/Source/ConcertSyncClient/Private/ConcertSyncClientUtil.cpp
--- a/Engine/Plugins/Developer/Concert/ConcertSync/ConcertSyncClient/Source/ConcertSyncClient/Private/ConcertSyncClientUtil.cpp
+++ b/Engine/Plugins/Developer/Concert/ConcertSync/ConcertSyncClient/Source/ConcertSyncClient/Private/ConcertSyncClientUtil.cpp
@@ -55,7 +55,32 @@ static TAutoConsoleVariable<int32> CVarDelayApplyingTransactionsWaitTimeout(
 	TEXT("Concert.DelayTransactionsWhileEditingTimeout"), 5,
 	TEXT("When Concert.DelayTransactionsWhileEditing is enabled we make sure the user has not been idle too long to prevent transactions from stacking up. The timeout value is specified in seconds."));
 
-bool UserIsEditing()
+bool IsUserEditing()
+{
+	static FName SEditableTextType(TEXT("SEditableText"));
+	static FName SMultiLineEditableTextType(TEXT("SMultiLineEditableText"));
+
+#if WITH_EDITOR
+	if (GUnrealEd)
+	{
+		auto IsUserEditingWidget = [] () {
+			bool bIsEditing = false;
+			FSlateApplication& App = FSlateApplication::Get();
+			App.ForEachUser([&bIsEditing](FSlateUser& User) {
+				TSharedPtr<SWidget> FocusedWidget = User.GetFocusedWidget();
+				bool bTextWidgetHasFocus = FocusedWidget && (FocusedWidget->GetType() == SEditableTextType || FocusedWidget->GetType() == SMultiLineEditableTextType);
+				bIsEditing |= bTextWidgetHasFocus;
+			});
+			return bIsEditing;
+		};
+		return GUnrealEd->IsUserInteracting() || IsUserEditingWidget();
+	}
+#endif
+
+	return false;
+}
+
+bool ShouldDelayTransaction()
 {
 #if WITH_EDITOR
 	if (CVarDelayApplyingTransactionsWhileEditing.GetValueOnAnyThread() > 0)
@@ -63,16 +88,10 @@ bool UserIsEditing()
 		static FName SEditableTextType(TEXT("SEditableText"));
 		static FName SMultiLineEditableTextType(TEXT("SMultiLineEditableText"));
 
-		bool bEditable = false;
-		FSlateApplication& App = FSlateApplication::Get();
-		App.ForEachUser([&bEditable](FSlateUser& User) {
-			TSharedPtr<SWidget> FocusedWidget = User.GetFocusedWidget();
-
-			bool bCanEdit = FocusedWidget && (FocusedWidget->GetType() == SEditableTextType || FocusedWidget->GetType() == SMultiLineEditableTextType);
-			bEditable |= bCanEdit;
-		});
-		if (bEditable)
+		const bool bIsEditing = IsUserEditing();
+		if (bIsEditing)
 		{
+			FSlateApplication& App = FSlateApplication::Get();
 			double LastUpdateTime = App.GetLastUserInteractionTime();
 			double Duration = App.GetCurrentTime() - LastUpdateTime;
 			if (static_cast<int32>(Duration) > CVarDelayApplyingTransactionsWaitTimeout.GetValueOnAnyThread())
@@ -80,7 +99,7 @@ bool UserIsEditing()
 				return false;
 			}
 		}
-		return bEditable;
+		return bIsEditing;
 	}
 #endif
 	return false;
diff --git a/Engine/Plugins/Developer/Concert/ConcertSync/ConcertSyncClient/Source/ConcertSyncClient/Private/ConcertSyncClientUtil.h b/Engine/Plugins/Developer/Concert/ConcertSync/ConcertSyncClient/Source/ConcertSyncClient/Private/ConcertSyncClientUtil.h
--- a/Engine/Plugins/Developer/Concert/ConcertSync/ConcertSyncClient/Source/ConcertSyncClient/Private/ConcertSyncClientUtil.h
+++ b/Engine/Plugins/Developer/Concert/ConcertSync/ConcertSyncClient/Source/ConcertSyncClient/Private/ConcertSyncClientUtil.h
@@ -67,8 +67,10 @@
 		EGetObjectResultFlags Flags;
 	};
 
-	bool UserIsEditing();
+	bool IsUserEditing();
 
+	bool ShouldDelayTransaction();
+
 	bool CanPerformBlockingAction(const bool bBlockDuringInteraction = true);
 
 	void UpdatePendingKillState(UObject* InObj, const bool bIsPendingKill);
diff --git a/Engine/Plugins/Tests/EditorTests/Source/EditorTests/Private/UnrealEd/ObjectToolsTests.cpp b/Engine/Plugins/Tests/EditorTests/Source/EditorTests/Private/UnrealEd/ObjectToolsTests.cpp
--- a/Engine/Plugins/Tests/EditorTests/Source/EditorTests/Private/UnrealEd/ObjectToolsTests.cpp
+++ b/Engine/Plugins/Tests/EditorTests/Source/EditorTests/Private/UnrealEd/ObjectToolsTests.cpp
@@ -155,7 +155,7 @@ protected:
 		}
 
 		// Validate that the transactor works as expected otherwise the test would be broken
-		check(GEditor->Trans->IsObjectInTransationBuffer(InReferencer));
+		check(GEditor->Trans->IsObjectInTransactionBuffer(InReferencer));
 
 		// The reference has been set to nullptr, only the transactor should now have a reference
 		TestGatherObjectReferencersForDeletion(
@@ -299,4 +299,4 @@ bool FObjectToolsTests_ForceDeleteObjects::RunTest(const FString&)
 	);
 
 	return true;
-}
\ No newline at end of file
+}
diff --git a/Engine/Source/Editor/UnrealEd/Classes/Editor/EditorEngine.h b/Engine/Source/Editor/UnrealEd/Classes/Editor/EditorEngine.h
--- a/Engine/Source/Editor/UnrealEd/Classes/Editor/EditorEngine.h
+++ b/Engine/Source/Editor/UnrealEd/Classes/Editor/EditorEngine.h
@@ -412,6 +412,14 @@ public:
 	UPROPERTY()
 	uint32 bRequestEndPlayMapQueued:1;
 
+	/** True if we should ignore noting any changes to selection on undo/redo */
+	UPROPERTY()
+	uint32 bIgnoreSelectionChange:1;
+
+	/** True if we should suspend notifying clients post undo/redo */
+	UPROPERTY()
+	uint32 bSuspendBroadcastPostUndoRedo:1;
+
 	/** True if we should not display notifications about undo/redo */
 	UPROPERTY()
 	uint32 bSquelchTransactionNotification:1;
@@ -427,7 +435,7 @@ public:
 	/** When Simulating In Editor, a pointer to the original (non-simulating) editor world */
 	UPROPERTY()
 	TObjectPtr<class UWorld> EditorWorld;
-	
+
 	/** When Simulating In Editor, an array of all actors that were selected when it began*/
 	UPROPERTY()
 	TArray<TWeakObjectPtr<class AActor> > ActorsThatWereSelected;
diff --git a/Engine/Source/Editor/UnrealEd/Classes/Editor/TransBuffer.h b/Engine/Source/Editor/UnrealEd/Classes/Editor/TransBuffer.h
--- a/Engine/Source/Editor/UnrealEd/Classes/Editor/TransBuffer.h
+++ b/Engine/Source/Editor/UnrealEd/Classes/Editor/TransBuffer.h
@@ -165,7 +165,7 @@ public:
 	virtual bool DisableObjectSerialization() override;
 	virtual bool IsObjectSerializationEnabled() override { return DisallowObjectSerialization == 0; }
 	virtual void SetPrimaryUndoObject( UObject* Object ) override;
-	virtual bool IsObjectInTransationBuffer( const UObject* Object ) const override;
+	virtual bool IsObjectInTransactionBuffer( const UObject* Object ) const override;
 	virtual bool IsObjectTransacting(const UObject* Object) const override;
 	virtual bool ContainsPieObjects() const override;
 	virtual bool IsActive() override
diff --git a/Engine/Source/Editor/UnrealEd/Classes/Editor/Transactor.h b/Engine/Source/Editor/UnrealEd/Classes/Editor/Transactor.h
--- a/Engine/Source/Editor/UnrealEd/Classes/Editor/Transactor.h
+++ b/Engine/Source/Editor/UnrealEd/Classes/Editor/Transactor.h
@@ -966,7 +966,11 @@ class UNREALED_API UTransactor : public UObject
 	virtual void SetPrimaryUndoObject( UObject* Object ) PURE_VIRTUAL(UTransactor::SetPrimaryUndoObject,);
 
 	/** Checks if a specific object is referenced by the transaction buffer */
-	virtual bool IsObjectInTransationBuffer( const UObject* Object ) const { return false; }
+	UE_DEPRECATED(5.1, "IsObjectInTransationBuffer is deprecated, please use IsObjectInTransactionBuffer instead.")
+	virtual bool IsObjectInTransationBuffer( const UObject* Object ) const { return IsObjectInTransactionBuffer(Object); }
+
+	/** Checks if a specific object is referenced by the transaction buffer */
+	virtual bool IsObjectInTransactionBuffer( const UObject* Object ) const { return false; }
 
 	/** Checks if a specific object is in the transaction currently underway */
 	virtual bool IsObjectTransacting(const UObject* Object) const PURE_VIRTUAL(UTransactor::IsObjectTransacting, return false;);
diff --git a/Engine/Source/Editor/UnrealEd/Private/EditorEngine.cpp b/Engine/Source/Editor/UnrealEd/Private/EditorEngine.cpp
--- a/Engine/Source/Editor/UnrealEd/Private/EditorEngine.cpp
+++ b/Engine/Source/Editor/UnrealEd/Private/EditorEngine.cpp
@@ -392,8 +392,10 @@ UEditorEngine::UEditorEngine(const FObjectInitializer& ObjectInitializer)
 				SetPreviewPlatform(FPreviewPlatformInfo(GMaxRHIFeatureLevel), false);
 			}
 		});
-		
+
 	bNotifyUndoRedoSelectionChange = true;
+	bIgnoreSelectionChange = false;
+	bSuspendBroadcastPostUndoRedo = false;
 
 	EditorWorldExtensionsManager = nullptr;
 
diff --git a/Engine/Source/Editor/UnrealEd/Private/EditorServer.cpp b/Engine/Source/Editor/UnrealEd/Private/EditorServer.cpp
--- a/Engine/Source/Editor/UnrealEd/Private/EditorServer.cpp
+++ b/Engine/Source/Editor/UnrealEd/Private/EditorServer.cpp
@@ -1198,14 +1198,20 @@ void UEditorEngine::HandleTransactorBeforeRedoUndo(const FTransactionContext& Tr
 
 void UEditorEngine::HandleTransactorRedoUndo(const FTransactionContext& TransactionContext, bool Succeeded, bool WasUndo)
 {
-	NoteSelectionChange(bNotifyUndoRedoSelectionChange);
+	if (!bIgnoreSelectionChange)
+	{
+		NoteSelectionChange(bNotifyUndoRedoSelectionChange);
+	}
 	PostUndo(Succeeded);
 
 	// Broadcast only if you have an actual transaction context
 	if (Succeeded)
 	{
 		check(CurrentUndoRedoContext.OuterOperationId.IsValid() && CurrentUndoRedoContext.OperationDepth > 0);
-		BroadcastPostUndoRedo(TransactionContext, WasUndo);
+		if (!bSuspendBroadcastPostUndoRedo)
+		{
+			BroadcastPostUndoRedo(TransactionContext, WasUndo);
+		}
 
 		if (--CurrentUndoRedoContext.OperationDepth == 0)
 		{
@@ -1450,7 +1456,7 @@ FText UEditorEngine::GetTransactionName() const
 
 bool UEditorEngine::IsObjectInTransactionBuffer( const UObject* Object ) const
 {
-	return Trans && Trans->IsObjectInTransationBuffer(Object);
+	return Trans && Trans->IsObjectInTransactionBuffer(Object);
 }
 
 bool UEditorEngine::Map_Select( UWorld* InWorld, const TCHAR* Str, FOutputDevice& Ar)
diff --git a/Engine/Source/Editor/UnrealEd/Private/EditorTransaction.cpp b/Engine/Source/Editor/UnrealEd/Private/EditorTransaction.cpp
--- a/Engine/Source/Editor/UnrealEd/Private/EditorTransaction.cpp
+++ b/Engine/Source/Editor/UnrealEd/Private/EditorTransaction.cpp
@@ -1704,7 +1704,7 @@ void UTransBuffer::SetPrimaryUndoObject(UObject* PrimaryObject)
 	}
 }
 
-bool UTransBuffer::IsObjectInTransationBuffer( const UObject* Object ) const
+bool UTransBuffer::IsObjectInTransactionBuffer( const UObject* Object ) const
 {
 	TArray<UObject*> TransactionObjects;
 	for( const TSharedRef<FTransaction>& Transaction : UndoBuffer )
