Remove unacceptable words from ConcertSyncClient and RemoteSession

diff --git a/Engine/Plugins/Developer/Concert/ConcertSync/ConcertSyncClient/Source/ConcertSyncClient/Private/ConcertClientSequencerManager.cpp b/Engine/Plugins/Developer/Concert/ConcertSync/ConcertSyncClient/Source/ConcertSyncClient/Private/ConcertClientSequencerManager.cpp
--- a/Engine/Plugins/Developer/Concert/ConcertSync/ConcertSyncClient/Source/ConcertSyncClient/Private/ConcertClientSequencerManager.cpp
+++ b/Engine/Plugins/Developer/Concert/ConcertSync/ConcertSyncClient/Source/ConcertSyncClient/Private/ConcertClientSequencerManager.cpp
@@ -768,7 +768,7 @@ void FConcertClientSequencerManager::ApplyEventToSequencers(const FConcertSequen
 		check(MovieScene);
 		MovieScene->SetPlaybackRange(EventState.PlaybackRange, false);
 
-		// If the event is coming from a sequencer that is playing back, we are a slave to its updates until it stops
+		// If the event is coming from a sequencer that is playing back, we are an agent to its updates until it stops
 		// We also apply any latency compensation when playing back
 		if (EventState.PlayerStatus == EConcertMovieScenePlayerStatus::Playing)
 		{
@@ -913,7 +913,7 @@ void FConcertClientSequencerManager::ApplyEventToPlayers(const FConcertSequencer
 
 		UE::Private::ConcertClientSequencerManager::ApplyPlayRangeToPlayer(Player, EventState.PlaybackRange);
 
-		// If the event is coming from a sequencer that is playing back, we are a slave to its updates until it stops
+		// If the event is coming from a sequencer that is playing back, we are an agent to its updates until it stops
 		// We also apply any latency compensation when playing back
 		if (EventState.PlayerStatus == EConcertMovieScenePlayerStatus::Playing)
 		{
diff --git a/Engine/Plugins/Experimental/RemoteSession/Source/RemoteSession/Private/Channels/RemoteSessionImageChannel.cpp b/Engine/Plugins/Experimental/RemoteSession/Source/RemoteSession/Private/Channels/RemoteSessionImageChannel.cpp
--- a/Engine/Plugins/Experimental/RemoteSession/Source/RemoteSession/Private/Channels/RemoteSessionImageChannel.cpp
+++ b/Engine/Plugins/Experimental/RemoteSession/Source/RemoteSession/Private/Channels/RemoteSessionImageChannel.cpp
@@ -17,9 +17,9 @@
 
 
 
-static int32 QualityMasterSetting = 0;
+static int32 QualityMainSetting = 0;
 static FAutoConsoleVariableRef CVarQualityOverride(
-	TEXT("remote.quality"), QualityMasterSetting,
+	TEXT("remote.quality"), QualityMainSetting,
 	TEXT("Sets quality (1-100)"),
 	ECVF_Default);
 
@@ -71,7 +71,7 @@ void FRemoteSessionImageChannel::FImageSender::SendRawImageToClients(int32 Width
 
 			ImageWrapper->SetRaw(ImageData, AllocatedImageDataSize, Width, Height, ERGBFormat::BGRA, 8, BytesPerRow);
 
-			const int32 CurrentQuality = QualityMasterSetting > 0 ? QualityMasterSetting : CompressQuality.Load();
+			const int32 CurrentQuality = QualityMainSetting > 0 ? QualityMainSetting : CompressQuality.Load();
 			const TArray64<uint8> JPGData = ImageWrapper->GetCompressed(CurrentQuality);
 
 			TBackChannelSharedPtr<FBackChannelOSCMessage> Msg = MakeShared<FBackChannelOSCMessage, ESPMode::ThreadSafe>(TEXT("/Screen"));
diff --git a/Engine/Plugins/Experimental/RemoteSession/Source/RemoteSession/Private/ImageProviders/RemoteSessionFrameBufferImageProvider.cpp b/Engine/Plugins/Experimental/RemoteSession/Source/RemoteSession/Private/ImageProviders/RemoteSessionFrameBufferImageProvider.cpp
--- a/Engine/Plugins/Experimental/RemoteSession/Source/RemoteSession/Private/ImageProviders/RemoteSessionFrameBufferImageProvider.cpp
+++ b/Engine/Plugins/Experimental/RemoteSession/Source/RemoteSession/Private/ImageProviders/RemoteSessionFrameBufferImageProvider.cpp
@@ -13,9 +13,9 @@
 #include "Widgets/SWindow.h"
 
 
-static int32 FramerateMasterSetting = 0;
+static int32 FramerateSetting = 0;
 static FAutoConsoleVariableRef CVarFramerateOverride(
-	TEXT("remote.framerate"), FramerateMasterSetting,
+	TEXT("remote.framerate"), FramerateSetting,
 	TEXT("Sets framerate"),
 	ECVF_Default);
 
@@ -64,7 +64,7 @@ void FRemoteSessionFrameBufferImageProvider::ReleaseFrameGrabber()
 void FRemoteSessionFrameBufferImageProvider::SetCaptureFrameRate(int32 InFramerate)
 {
 	// Set our framerate and quality cvars, if the user hasn't modified them
-	if (FramerateMasterSetting == 0)
+	if (FramerateSetting == 0)
 	{
 		CVarFramerateOverride->Set(InFramerate);
 	}
@@ -130,7 +130,7 @@ void FRemoteSessionFrameBufferImageProvider::Tick(const float InDeltaTime)
 
 	const double TimeNow = FPlatformTime::Seconds();
 	const double ElapsedImageTimeMS = (TimeNow - LastSentImageTime) * 1000;
-	const int32 DesiredFrameTimeMS = 1000 / FramerateMasterSetting;
+	const int32 DesiredFrameTimeMS = 1000 / FramerateSetting;
 
 	if (ElapsedImageTimeMS < DesiredFrameTimeMS)
 	{
