Adding additional includes in preparation for removing a global include

#rnx

diff --git a/Engine/Plugins/Developer/Concert/ConcertSync/ConcertSyncClient/Source/ConcertSyncClient/Private/ConcertClientTransactionBridge.cpp b/Engine/Plugins/Developer/Concert/ConcertSync/ConcertSyncClient/Source/ConcertSyncClient/Private/ConcertClientTransactionBridge.cpp
--- a/Engine/Plugins/Developer/Concert/ConcertSync/ConcertSyncClient/Source/ConcertSyncClient/Private/ConcertClientTransactionBridge.cpp
+++ b/Engine/Plugins/Developer/Concert/ConcertSync/ConcertSyncClient/Source/ConcertSyncClient/Private/ConcertClientTransactionBridge.cpp
@@ -7,6 +7,7 @@
 
 #include "ConcertTransactionEvents.h"
 #include "IConcertClientTransactionBridge.h"
+#include "Misc/ITransactionObjectAnnotation.h"
 #include "Misc/PackageName.h"
 #include "Misc/CoreDelegates.h"
 #include "HAL/IConsoleManager.h"
diff --git a/Engine/Source/Editor/UnrealEd/Private/EditorTransaction.cpp b/Engine/Source/Editor/UnrealEd/Private/EditorTransaction.cpp
--- a/Engine/Source/Editor/UnrealEd/Private/EditorTransaction.cpp
+++ b/Engine/Source/Editor/UnrealEd/Private/EditorTransaction.cpp
@@ -10,6 +10,7 @@
 #include "Engine/Level.h"
 #include "Components/ActorComponent.h"
 #include "Model.h"
+#include "Misc/ITransactionObjectAnnotation.h"
 #include "Editor/Transactor.h"
 #include "Editor/TransBuffer.h"
 #include "Components/ModelComponent.h"
diff --git a/Engine/Source/Editor/UnrealEd/Private/LevelEditorViewport.cpp b/Engine/Source/Editor/UnrealEd/Private/LevelEditorViewport.cpp
--- a/Engine/Source/Editor/UnrealEd/Private/LevelEditorViewport.cpp
+++ b/Engine/Source/Editor/UnrealEd/Private/LevelEditorViewport.cpp
@@ -5,6 +5,7 @@
 #include "Materials/MaterialInterface.h"
 #include "Modules/ModuleManager.h"
 #include "Misc/PackageName.h"
+#include "Misc/ITransaction.h"
 #include "Framework/Application/SlateApplication.h"
 #include "Styling/AppStyle.h"
 #include "Components/MeshComponent.h"
