Cleaned up the exported property utils in Concert

#rnx

diff --git a/Engine/Plugins/Developer/Concert/ConcertSync/ConcertSyncClient/Source/ConcertSyncClient/Private/ConcertClientTransactionBridge.cpp b/Engine/Plugins/Developer/Concert/ConcertSync/ConcertSyncClient/Source/ConcertSyncClient/Private/ConcertClientTransactionBridge.cpp
--- a/Engine/Plugins/Developer/Concert/ConcertSync/ConcertSyncClient/Source/ConcertSyncClient/Private/ConcertClientTransactionBridge.cpp
+++ b/Engine/Plugins/Developer/Concert/ConcertSync/ConcertSyncClient/Source/ConcertSyncClient/Private/ConcertClientTransactionBridge.cpp
@@ -673,7 +673,7 @@ void FConcertClientTransactionBridge::HandleObjectTransacted(UObject* InObject,
 	const FName NewObjectName = InTransactionEvent.HasNameChange() ? InObject->GetFName() : FName();
 	const FName NewObjectOuterPathName = (InTransactionEvent.HasOuterChange() && InObject->GetOuter()) ? FName(*InObject->GetOuter()->GetPathName()) : FName();
 	const FName NewObjectExternalPackageName = (InTransactionEvent.HasExternalPackageChange() && InObject->GetExternalPackage()) ? InObject->GetExternalPackage()->GetFName() : FName();
-	const TArray<FName> RootPropertyNames = ConcertSyncClientUtil::GetRootProperties(InTransactionEvent.GetChangedProperties());
+	const TArray<const FProperty*> ExportedProperties = ConcertSyncClientUtil::GetExportedProperties(InObject->GetClass(), InTransactionEvent.GetChangedProperties(), bIncludeEditorOnlyProperties);
 	TSharedPtr<ITransactionObjectAnnotation> TransactionAnnotation = InTransactionEvent.GetAnnotation();
 
 	// Track which packages were changed
@@ -721,26 +721,20 @@ void FConcertClientTransactionBridge::HandleObjectTransacted(UObject* InObject,
 			}
 
 			// Find or add an update for each property
-			for (const FName& RootPropertyName : RootPropertyNames)
+			for (const FProperty* ExportedProperty : ExportedProperties)
 			{
-				FProperty* RootProperty = ConcertSyncClientUtil::GetExportedProperty(InObject->GetClass(), RootPropertyName, bIncludeEditorOnlyProperties);
-				if (!RootProperty)
+				FConcertSerializedPropertyData* PropertyDataPtr = ObjectUpdatePtr->PropertyDatas.FindByPredicate([ExportedProperty](FConcertSerializedPropertyData& PropertyData)
 				{
-					continue;
-				}
-
-				FConcertSerializedPropertyData* PropertyDataPtr = ObjectUpdatePtr->PropertyDatas.FindByPredicate([&RootPropertyName](FConcertSerializedPropertyData& PropertyData)
-				{
-					return RootPropertyName == PropertyData.PropertyName;
+					return ExportedProperty->GetFName() == PropertyData.PropertyName;
 				});
 				if (!PropertyDataPtr)
 				{
 					PropertyDataPtr = &ObjectUpdatePtr->PropertyDatas.AddDefaulted_GetRef();
-					PropertyDataPtr->PropertyName = RootPropertyName;
+					PropertyDataPtr->PropertyName = ExportedProperty->GetFName();
 				}
 
 				PropertyDataPtr->SerializedData.Reset();
-				ConcertSyncClientUtil::SerializeProperty(nullptr, InObject, RootProperty, bIncludeEditorOnlyProperties, PropertyDataPtr->SerializedData);
+				ConcertSyncClientUtil::SerializeProperty(nullptr, InObject, ExportedProperty, bIncludeEditorOnlyProperties, PropertyDataPtr->SerializedData);
 			}
 		}
 	}
@@ -773,31 +767,20 @@ void FConcertClientTransactionBridge::HandleObjectTransacted(UObject* InObject,
 		{
 			// The 'non-property changes' refers to custom data added by a deriver UObject before and/or after the standard serialized data. Since this is a custom
 			// data format, we don't know what changed, call the object to re-serialize this part, but still send the delta for the generic reflected properties (in RootPropertyNames).
-			ConcertSyncClientUtil::SerializeObject(&OngoingTransaction.FinalizedData.FinalizedLocalIdentifierTable, InObject, &RootPropertyNames, bIncludeEditorOnlyProperties, ObjectUpdate.ObjectData.SerializedData);
+			ConcertSyncClientUtil::SerializeObject(&OngoingTransaction.FinalizedData.FinalizedLocalIdentifierTable, InObject, &ExportedProperties, bIncludeEditorOnlyProperties, ObjectUpdate.ObjectData.SerializedData);
 
 			// Track which properties changed. Not used to apply the transaction on the receiving side, the object specific serialization function will be used for that, but
 			// to be able to display, in the transaction detail view, which 'properties' changed in the transaction as transaction data is otherwise opaque to UI.
-			for (const FName& RootPropertyName : RootPropertyNames)
+			for (const FProperty* ExportedProperty : ExportedProperties)
 			{
-				if (FProperty* RootProperty = ConcertSyncClientUtil::GetExportedProperty(InObject->GetClass(), RootPropertyName, bIncludeEditorOnlyProperties))
-				{
-					FConcertSerializedPropertyData& PropertyData = ObjectUpdate.PropertyDatas.AddDefaulted_GetRef();
-					PropertyData.PropertyName = RootPropertyName;
-				}
+				FConcertSerializedPropertyData& PropertyData = ObjectUpdate.PropertyDatas.AddDefaulted_GetRef();
+				PropertyData.PropertyName = ExportedProperty->GetFName();
 			}
 		}
 		else // Its possible to optimize the transaction payload, only sending a 'delta' update.
 		{
 			// Only send properties that changed. The receiving side will 'patch' the object using the reflection system. The specific object serialization function will NOT be called.
-			for (const FName& RootPropertyName : RootPropertyNames)
-			{
-				if (FProperty* RootProperty = ConcertSyncClientUtil::GetExportedProperty(InObject->GetClass(), RootPropertyName, bIncludeEditorOnlyProperties))
-				{
-					FConcertSerializedPropertyData& PropertyData = ObjectUpdate.PropertyDatas.AddDefaulted_GetRef();
-					PropertyData.PropertyName = RootPropertyName;
-					ConcertSyncClientUtil::SerializeProperty(&OngoingTransaction.FinalizedData.FinalizedLocalIdentifierTable, InObject, RootProperty, bIncludeEditorOnlyProperties, PropertyData.SerializedData);
-				}
-			}
+			ConcertSyncClientUtil::SerializeProperties(&OngoingTransaction.FinalizedData.FinalizedLocalIdentifierTable, InObject, ExportedProperties, bIncludeEditorOnlyProperties, ObjectUpdate.PropertyDatas);
 		}
 	}
 }
diff --git a/Engine/Plugins/Developer/Concert/ConcertSync/ConcertSyncClient/Source/ConcertSyncClient/Private/ConcertSyncClientUtil.cpp b/Engine/Plugins/Developer/Concert/ConcertSync/ConcertSyncClient/Source/ConcertSyncClient/Private/ConcertSyncClientUtil.cpp
--- a/Engine/Plugins/Developer/Concert/ConcertSync/ConcertSyncClient/Source/ConcertSyncClient/Private/ConcertSyncClientUtil.cpp
+++ b/Engine/Plugins/Developer/Concert/ConcertSync/ConcertSyncClient/Source/ConcertSyncClient/Private/ConcertSyncClientUtil.cpp
@@ -312,38 +312,25 @@ FGetObjectResult GetObject(const FConcertObjectId& InObjectId, const FName InNew
 	return FGetObjectResult();
 }
 
-bool ImportPropertyData(const FConcertLocalIdentifierTable* InLocalIdentifierTable, const FConcertSyncWorldRemapper& InWorldRemapper, const FConcertSessionVersionInfo* InVersionInfo, UObject* InObj, const FName InPropertyName, const TArray<uint8>& InSerializedData)
+TArray<const FProperty*> GetExportedProperties(const UStruct* InStruct, const TArray<FName>& InPropertyNames, const bool InIncludeEditorOnlyData)
 {
-	FProperty* Prop = InObj->GetClass()->FindPropertyByName(InPropertyName);
-	if (Prop)
-	{
-		FConcertSyncObjectReader ObjectReader(InLocalIdentifierTable, InWorldRemapper, InVersionInfo, InObj, InSerializedData);
-		ObjectReader.SerializeProperty(Prop, InObj);
-		return !ObjectReader.GetError();
-	}
-
-	return false;
-}
+	TArray<const FProperty*> ExportedProperties;
+	ExportedProperties.Reserve(InPropertyNames.Num());
 
-TArray<FName> GetRootProperties(const TArray<FName>& InChangedProperties)
-{
-	TArray<FName> RootProperties;
-	RootProperties.Reserve(InChangedProperties.Num());
-	for (const FName& PropertyChainName : InChangedProperties)
+	for (const FName& PropertyName : InPropertyNames)
 	{
-		// Only care about the root property in the chain
-		TArray<FString> PropertyChainNames;
-		PropertyChainName.ToString().ParseIntoArray(PropertyChainNames, TEXT("."));
-		check(PropertyChainNames.Num() >= 1);
-
-		RootProperties.AddUnique(*PropertyChainNames[0]);
+		if (const FProperty* Property = GetExportedProperty(InStruct, PropertyName, InIncludeEditorOnlyData))
+		{
+			ExportedProperties.Add(Property);
+		}
 	}
-	return RootProperties;
+
+	return ExportedProperties;
 }
 
-FProperty* GetExportedProperty(const UStruct* InStruct, const FName InPropertyName, const bool InIncludeEditorOnlyData)
+const FProperty* GetExportedProperty(const UStruct* InStruct, const FName InPropertyName, const bool InIncludeEditorOnlyData)
 {
-	if (FProperty* Property = FindFProperty<FProperty>(InStruct, InPropertyName))
+	if (const FProperty* Property = FindFProperty<FProperty>(InStruct, InPropertyName))
 	{
 		if (ConcertSyncUtil::CanExportProperty(Property, InIncludeEditorOnlyData))
 		{
@@ -354,18 +341,13 @@ FProperty* GetExportedProperty(const UStruct* InStruct, const FName InPropertyNa
 	return nullptr;
 }
 
-void SerializeProperties(FConcertLocalIdentifierTable* InLocalIdentifierTable, const UObject* InObject, const TArray<FName>& InChangedProperties, const bool InIncludeEditorOnlyData, TArray<FConcertSerializedPropertyData>& OutPropertyDatas)
+void SerializeProperties(FConcertLocalIdentifierTable* InLocalIdentifierTable, const UObject* InObject, const TArray<const FProperty*>& InProperties, const bool InIncludeEditorOnlyData, TArray<FConcertSerializedPropertyData>& OutPropertyDatas)
 {
-	const TArray<FName> RootPropertyNames = GetRootProperties(InChangedProperties);
-	for (const FName& RootPropertyName : RootPropertyNames)
+	for (const FProperty* Property : InProperties)
 	{
-		FProperty* RootProperty = GetExportedProperty(InObject->GetClass(), RootPropertyName, InIncludeEditorOnlyData);
-		if (RootProperty)
-		{
-			FConcertSerializedPropertyData& PropertyData = OutPropertyDatas.AddDefaulted_GetRef();
-			PropertyData.PropertyName = RootProperty->GetFName();
-			SerializeProperty(InLocalIdentifierTable, InObject, RootProperty, InIncludeEditorOnlyData, PropertyData.SerializedData);
-		}
+		FConcertSerializedPropertyData& PropertyData = OutPropertyDatas.AddDefaulted_GetRef();
+		PropertyData.PropertyName = Property->GetFName();
+		SerializeProperty(InLocalIdentifierTable, InObject, Property, InIncludeEditorOnlyData, PropertyData.SerializedData);
 	}
 }
 
@@ -374,15 +356,15 @@ void SerializeProperty(FConcertLocalIdentifierTable* InLocalIdentifierTable, con
 	bool bSkipAssets = false; // TODO: Handle asset updates
 
 	FConcertSyncObjectWriter ObjectWriter(InLocalIdentifierTable, (UObject*)InObject, OutSerializedData, InIncludeEditorOnlyData, bSkipAssets);
-	ObjectWriter.SerializeProperty((FProperty*)InProperty, (UObject*)InObject);
+	ObjectWriter.SerializeProperty(InProperty, InObject);
 }
 
-void SerializeObject(FConcertLocalIdentifierTable* InLocalIdentifierTable, const UObject* InObject, const TArray<FName>* InChangedProperties, const bool InIncludeEditorOnlyData, TArray<uint8>& OutSerializedData)
+void SerializeObject(FConcertLocalIdentifierTable* InLocalIdentifierTable, const UObject* InObject, const TArray<const FProperty*>* InProperties, const bool InIncludeEditorOnlyData, TArray<uint8>& OutSerializedData)
 {
 	bool bSkipAssets = false; // TODO: Handle asset updates
 
 	FConcertSyncObjectWriter ObjectWriter(InLocalIdentifierTable, (UObject*)InObject, OutSerializedData, InIncludeEditorOnlyData, bSkipAssets);
-	ObjectWriter.SerializeObject((UObject*)InObject, InChangedProperties);
+	ObjectWriter.SerializeObject(InObject, InProperties);
 }
 
 
diff --git a/Engine/Plugins/Developer/Concert/ConcertSync/ConcertSyncClient/Source/ConcertSyncClient/Private/ConcertSyncClientUtil.h b/Engine/Plugins/Developer/Concert/ConcertSync/ConcertSyncClient/Source/ConcertSyncClient/Private/ConcertSyncClientUtil.h
--- a/Engine/Plugins/Developer/Concert/ConcertSync/ConcertSyncClient/Source/ConcertSyncClient/Private/ConcertSyncClientUtil.h
+++ b/Engine/Plugins/Developer/Concert/ConcertSync/ConcertSyncClient/Source/ConcertSyncClient/Private/ConcertSyncClientUtil.h
@@ -81,17 +81,15 @@ namespace ConcertSyncClientUtil
 
 	FGetObjectResult GetObject(const FConcertObjectId& InObjectId, const FName InNewName, const FName InNewOuterPath, const FName InNewPackageName, const bool bAllowCreate);
 	
-	bool ImportPropertyData(const FConcertLocalIdentifierTable* InLocalIdentifierTable, const FConcertSyncWorldRemapper& InWorldRemapper, const FConcertSessionVersionInfo* InVersionInfo, UObject* InObj, const FName InPropertyName, const TArray<uint8>& InSerializedData);
+	TArray<const FProperty*> GetExportedProperties(const UStruct* InStruct, const TArray<FName>& InPropertyNames, const bool InIncludeEditorOnlyData);
 
-	TArray<FName> GetRootProperties(const TArray<FName>& InChangedProperties);
+	const FProperty* GetExportedProperty(const UStruct* InStruct, const FName InPropertyName, const bool InIncludeEditorOnlyData);
 
-	FProperty* GetExportedProperty(const UStruct* InStruct, const FName InPropertyName, const bool InIncludeEditorOnlyData);
-
-	void SerializeProperties(FConcertLocalIdentifierTable* InLocalIdentifierTable, const UObject* InObject, const TArray<FName>& InChangedProperties, const bool InIncludeEditorOnlyData, TArray<FConcertSerializedPropertyData>& OutPropertyDatas);
+	void SerializeProperties(FConcertLocalIdentifierTable* InLocalIdentifierTable, const UObject* InObject, const TArray<const FProperty*>& InProperties, const bool InIncludeEditorOnlyData, TArray<FConcertSerializedPropertyData>& OutPropertyDatas);
 
 	void SerializeProperty(FConcertLocalIdentifierTable* InLocalIdentifierTable, const UObject* InObject, const FProperty* InProperty, const bool InIncludeEditorOnlyData, TArray<uint8>& OutSerializedData);
 
-	void SerializeObject(FConcertLocalIdentifierTable* InLocalIdentifierTable, const UObject* InObject, const TArray<FName>* InChangedProperties, const bool InIncludeEditorOnlyData, TArray<uint8>& OutSerializedData);
+	void SerializeObject(FConcertLocalIdentifierTable* InLocalIdentifierTable, const UObject* InObject, const TArray<const FProperty*>* InProperties, const bool InIncludeEditorOnlyData, TArray<uint8>& OutSerializedData);
 
 	void FlushPackageLoading(const FName InPackageName);
 
diff --git a/Engine/Plugins/Developer/Concert/ConcertSync/ConcertSyncCore/Source/ConcertSyncCore/Private/ConcertSyncArchives.cpp b/Engine/Plugins/Developer/Concert/ConcertSync/ConcertSyncCore/Source/ConcertSyncCore/Private/ConcertSyncArchives.cpp
--- a/Engine/Plugins/Developer/Concert/ConcertSync/ConcertSyncCore/Source/ConcertSyncCore/Private/ConcertSyncArchives.cpp
+++ b/Engine/Plugins/Developer/Concert/ConcertSync/ConcertSyncCore/Source/ConcertSyncCore/Private/ConcertSyncArchives.cpp
@@ -223,30 +223,30 @@ FConcertSyncObjectWriter::FConcertSyncObjectWriter(FConcertLocalIdentifierTable*
 {
 }
 
-void FConcertSyncObjectWriter::SerializeObject(UObject* InObject, const TArray<FName>* InPropertyNamesToWrite)
+void FConcertSyncObjectWriter::SerializeObject(const UObject* InObject, const TArray<const FProperty*>* InPropertiesToWrite)
 {
-	if (InPropertyNamesToWrite)
+	if (InPropertiesToWrite)
 	{
-		ShouldSkipPropertyFunc = [InObject, InPropertyNamesToWrite](const FProperty* InProperty) -> bool
+		ShouldSkipPropertyFunc = [InObject, InPropertiesToWrite](const FProperty* InProperty) -> bool
 		{
-			return InProperty->GetOwnerStruct() == InObject->GetClass() && !InPropertyNamesToWrite->Contains(InProperty->GetFName());
+			return !InPropertiesToWrite->Contains(InProperty);
 		};
 
-		InObject->Serialize(*this);
+		const_cast<UObject*>(InObject)->Serialize(*this);
 
 		ShouldSkipPropertyFunc = FShouldSkipPropertyFunc();
 	}
 	else
 	{
-		InObject->Serialize(*this);
+		const_cast<UObject*>(InObject)->Serialize(*this);
 	}
 }
 
-void FConcertSyncObjectWriter::SerializeProperty(FProperty* InProp, UObject* InObject)
+void FConcertSyncObjectWriter::SerializeProperty(const FProperty* InProp, const UObject* InObject)
 {
 	for (int32 Idx = 0; Idx < InProp->ArrayDim; ++Idx)
 	{
-		InProp->SerializeItem(FStructuredArchiveFromArchive(*this).GetSlot(), InProp->ContainerPtrToValuePtr<void>(InObject, Idx));
+		InProp->SerializeItem(FStructuredArchiveFromArchive(*this).GetSlot(), InProp->ContainerPtrToValuePtr<void>(const_cast<UObject*>(InObject), Idx));
 	}
 }
 
@@ -372,7 +372,7 @@ void FConcertSyncObjectReader::SerializeObject(UObject* InObject)
 	InObject->Serialize(*this);
 }
 
-void FConcertSyncObjectReader::SerializeProperty(FProperty* InProp, UObject* InObject)
+void FConcertSyncObjectReader::SerializeProperty(const FProperty* InProp, UObject* InObject)
 {
 	for (int32 Idx = 0; Idx < InProp->ArrayDim; ++Idx)
 	{
diff --git a/Engine/Plugins/Developer/Concert/ConcertSync/ConcertSyncCore/Source/ConcertSyncCore/Public/ConcertSyncArchives.h b/Engine/Plugins/Developer/Concert/ConcertSync/ConcertSyncCore/Source/ConcertSyncCore/Public/ConcertSyncArchives.h
--- a/Engine/Plugins/Developer/Concert/ConcertSync/ConcertSyncCore/Source/ConcertSyncCore/Public/ConcertSyncArchives.h
+++ b/Engine/Plugins/Developer/Concert/ConcertSync/ConcertSyncCore/Source/ConcertSyncCore/Public/ConcertSyncArchives.h
@@ -50,8 +50,8 @@ public:
 	FConcertSyncObjectWriter(FConcertLocalIdentifierTable* InLocalIdentifierTable, UObject* InObj, TArray<uint8>& OutBytes, const bool InIncludeEditorOnlyData, const bool InSkipAssets, const FConcertSyncRemapObjectPath& InRemapDelegate);
 	FConcertSyncObjectWriter(FConcertLocalIdentifierTable* InLocalIdentifierTable, UObject* InObj, TArray<uint8>& OutBytes, const bool InIncludeEditorOnlyData, const bool InSkipAssets);
 
-	void SerializeObject(UObject* InObject, const TArray<FName>* InPropertyNamesToWrite = nullptr);
-	void SerializeProperty(FProperty* InProp, UObject* InObject);
+	void SerializeObject(const UObject* InObject, const TArray<const FProperty*>* InPropertiesToWrite = nullptr);
+	void SerializeProperty(const FProperty* InProp, const UObject* InObject);
 
 	using FConcertIdentifierWriter::operator<<; // For visibility of the overloads we don't override
 
@@ -80,7 +80,7 @@ public:
 	FConcertSyncObjectReader(const FConcertLocalIdentifierTable* InLocalIdentifierTable, FConcertSyncWorldRemapper InWorldRemapper, const FConcertSessionVersionInfo* InVersionInfo, UObject* InObj, const TArray<uint8>& InBytes);
 
 	void SerializeObject(UObject* InObject);
-	void SerializeProperty(FProperty* InProp, UObject* InObject);
+	void SerializeProperty(const FProperty* InProp, UObject* InObject);
 
 	using FConcertIdentifierReader::operator<<; // For visibility of the overloads we don't override
 
