Prepare Concert to accept updates with bResetExisting set

 - Actor creation is no longer deferred so that managed components (those created by a constructor) are constructed prior to being searched for.
 - Update application (and Pre/PostEditChange/Undo) now runs in child->parent order. This ensures that components are updated prior to the actor construction script re-running, and matches how undo/redo orders things.

#rnx

diff --git a/Engine/Plugins/Developer/Concert/ConcertSync/ConcertSyncClient/Source/ConcertSyncClient/Private/ConcertClientTransactionBridge.cpp b/Engine/Plugins/Developer/Concert/ConcertSync/ConcertSyncClient/Source/ConcertSyncClient/Private/ConcertClientTransactionBridge.cpp
--- a/Engine/Plugins/Developer/Concert/ConcertSync/ConcertSyncClient/Source/ConcertSyncClient/Private/ConcertClientTransactionBridge.cpp
+++ b/Engine/Plugins/Developer/Concert/ConcertSync/ConcertSyncClient/Source/ConcertSyncClient/Private/ConcertClientTransactionBridge.cpp
@@ -210,43 +210,44 @@ struct FEditorTransactionNotification
 };
 #endif
 
-void ProcessTransactionEvent(const FConcertTransactionEventBase& InEvent, const FConcertSessionVersionInfo* InVersionInfo, const TArray<FName>& InPackagesToProcess, const FConcertLocalIdentifierTable* InLocalIdentifierTablePtr, const bool bIsSnapshot, const FConcertSyncWorldRemapper& WorldRemapper, TSet<AActor*>* DeletedActors)
+void ProcessTransactionEvent(const FConcertTransactionEventBase& InEvent, const FConcertSessionVersionInfo* InVersionInfo, const TArray<FName>& InPackagesToProcess, const FConcertLocalIdentifierTable* InLocalIdentifierTablePtr, const bool bIsSnapshot, const FConcertSyncWorldRemapper& WorldRemapper, TSet<AActor*>* DeletedActors, const bool bIncludeEditorOnlyProperties)
 {
 	// Transactions are applied in multiple-phases...
+	//	0) Sort the objects to be processed; creation needs to happen parent->child, and update needs to happen child->parent
 	//	1) Find or create all objects in the transaction (to handle object-interdependencies in the serialized data)
 	//	2) Notify all objects that they are about to be changed (via PreEditUndo)
 	//	3) Update the state of all objects
-	//	4) Notify all objects that they were changed (via PostEditUndo) - also finish spawning any new actors now that they have the correct state
+	//	4) Notify all objects that they were changed (via PostEditUndo)
 
 	// --------------------------------------------------------------------------------------------------------------------
-	// Phase 1
+	// Phase 0
 	// --------------------------------------------------------------------------------------------------------------------
-	bool bObjectsDeleted = false;
-	TArray<ConcertSyncClientUtil::FGetObjectResult, TInlineAllocator<8>> TransactionObjects;
-	TransactionObjects.AddDefaulted(InEvent.ExportedObjects.Num());
+	TArray<const FConcertExportedObject*, TInlineAllocator<8>> SortedExportedObjects;
 	{
-		// Sort the object list so that outers will be created before their child objects
-		typedef TTuple<int32, const FConcertExportedObject*> FConcertExportedIndexAndObject;
-		TArray<FConcertExportedIndexAndObject, TInlineAllocator<8>> SortedExportedObjects;
 		SortedExportedObjects.Reserve(InEvent.ExportedObjects.Num());
-		for (int32 ObjectIndex = 0; ObjectIndex < InEvent.ExportedObjects.Num(); ++ObjectIndex)
+		for (const FConcertExportedObject& ExportedObject : InEvent.ExportedObjects)
 		{
-			SortedExportedObjects.Add(MakeTuple(ObjectIndex, &InEvent.ExportedObjects[ObjectIndex]));
+			SortedExportedObjects.Add(&ExportedObject);
 		}
 
-		SortedExportedObjects.StableSort([](const FConcertExportedIndexAndObject& One, const FConcertExportedIndexAndObject& Two) -> bool
+		SortedExportedObjects.StableSort([](const FConcertExportedObject& One, const FConcertExportedObject& Two) -> bool
 		{
-			const FConcertExportedObject& OneObjectUpdate = *One.Value;
-			const FConcertExportedObject& TwoObjectUpdate = *Two.Value;
-			return OneObjectUpdate.ObjectPathDepth < TwoObjectUpdate.ObjectPathDepth;
+			return One.ObjectPathDepth < Two.ObjectPathDepth;
 		});
+	}
 
-		// Find or create each object, populating TransactionObjects in the original order (not the sorted order)
-		for (const FConcertExportedIndexAndObject& SortedExportedObjectsPair : SortedExportedObjects)
+	// --------------------------------------------------------------------------------------------------------------------
+	// Phase 1
+	// --------------------------------------------------------------------------------------------------------------------
+	bool bObjectsDeleted = false;
+	TArray<ConcertSyncClientUtil::FGetObjectResult, TInlineAllocator<8>> TransactionObjects;
+	{
+		// Find or create each object in parent->child order
+		TransactionObjects.AddDefaulted(SortedExportedObjects.Num());
+		for (int32 ObjectIndex = 0; ObjectIndex < SortedExportedObjects.Num(); ++ObjectIndex)
 		{
-			const int32 ObjectUpdateIndex = SortedExportedObjectsPair.Key;
-			const FConcertExportedObject& ObjectUpdate = *SortedExportedObjectsPair.Value;
-			ConcertSyncClientUtil::FGetObjectResult& TransactionObjectRef = TransactionObjects[ObjectUpdateIndex];
+			ConcertSyncClientUtil::FGetObjectResult& TransactionObjectRef = TransactionObjects[ObjectIndex];
+			const FConcertExportedObject& ObjectUpdate = *SortedExportedObjects[ObjectIndex];
 
 			// Is this object excluded? We exclude certain packages when re-applying live transactions on a package load
 			if (InPackagesToProcess.Num() > 0)
@@ -262,6 +263,12 @@ void ProcessTransactionEvent(const FConcertTransactionEventBase& InEvent, const
 			// Find or create the object
 			TransactionObjectRef = ConcertSyncClientUtil::GetObject(ObjectUpdate.ObjectId, ObjectUpdate.ObjectData.NewName, ObjectUpdate.ObjectData.NewOuterPathName, ObjectUpdate.ObjectData.NewExternalPackageName, ObjectUpdate.ObjectData.bAllowCreate);
 			bObjectsDeleted |= (ObjectUpdate.ObjectData.bIsPendingKill || TransactionObjectRef.NeedsGC());
+
+			if (ObjectUpdate.ObjectData.bResetExisting && !TransactionObjectRef.NewlyCreated())
+			{
+				check(TransactionObjectRef.Obj);
+				ConcertSyncUtil::ResetObjectPropertiesToArchetypeValues(TransactionObjectRef.Obj, bIncludeEditorOnlyProperties);
+			}
 		}
 	}
 
@@ -279,11 +286,11 @@ void ProcessTransactionEvent(const FConcertTransactionEventBase& InEvent, const
 	// --------------------------------------------------------------------------------------------------------------------
 #if WITH_EDITOR
 	TArray<TSharedPtr<ITransactionObjectAnnotation>, TInlineAllocator<8>> TransactionAnnotations;
-	TransactionAnnotations.AddDefaulted(InEvent.ExportedObjects.Num());
-	for (int32 ObjectIndex = 0; ObjectIndex < TransactionObjects.Num(); ++ObjectIndex)
+	TransactionAnnotations.AddDefaulted(SortedExportedObjects.Num());
+	for (int32 ObjectIndex = SortedExportedObjects.Num() - 1; ObjectIndex >= 0; --ObjectIndex)
 	{
 		const ConcertSyncClientUtil::FGetObjectResult& TransactionObjectRef = TransactionObjects[ObjectIndex];
-		const FConcertExportedObject& ObjectUpdate = InEvent.ExportedObjects[ObjectIndex];
+		const FConcertExportedObject& ObjectUpdate = *SortedExportedObjects[ObjectIndex];
 
 		UObject* TransactionObject = TransactionObjectRef.Obj;
 		if (!TransactionObject)
@@ -341,10 +348,10 @@ void ProcessTransactionEvent(const FConcertTransactionEventBase& InEvent, const
 	// --------------------------------------------------------------------------------------------------------------------
 	// Phase 3
 	// --------------------------------------------------------------------------------------------------------------------
-	for (int32 ObjectIndex = 0; ObjectIndex < TransactionObjects.Num(); ++ObjectIndex)
+	for (int32 ObjectIndex = SortedExportedObjects.Num() - 1; ObjectIndex >= 0; --ObjectIndex)
 	{
 		const ConcertSyncClientUtil::FGetObjectResult& TransactionObjectRef = TransactionObjects[ObjectIndex];
-		const FConcertExportedObject& ObjectUpdate = InEvent.ExportedObjects[ObjectIndex];
+		const FConcertExportedObject& ObjectUpdate = *SortedExportedObjects[ObjectIndex];
 
 		UObject* TransactionObject = TransactionObjectRef.Obj;
 		if (!TransactionObject)
@@ -378,10 +385,10 @@ void ProcessTransactionEvent(const FConcertTransactionEventBase& InEvent, const
 	// --------------------------------------------------------------------------------------------------------------------
 	// Phase 4
 	// --------------------------------------------------------------------------------------------------------------------
-	for (int32 ObjectIndex = 0; ObjectIndex < TransactionObjects.Num(); ++ObjectIndex)
+	for (int32 ObjectIndex = SortedExportedObjects.Num() - 1; ObjectIndex >= 0; --ObjectIndex)
 	{
 		const ConcertSyncClientUtil::FGetObjectResult& TransactionObjectRef = TransactionObjects[ObjectIndex];
-		const FConcertExportedObject& ObjectUpdate = InEvent.ExportedObjects[ObjectIndex];
+		const FConcertExportedObject& ObjectUpdate = *SortedExportedObjects[ObjectIndex];
 
 		UObject* TransactionObject = TransactionObjectRef.Obj;
 		if (!TransactionObject)
@@ -389,7 +396,6 @@ void ProcessTransactionEvent(const FConcertTransactionEventBase& InEvent, const
 			continue;
 		}
 
-		// Finish spawning any newly created actors
 		if (AActor* TransactionActor = Cast<AActor>(TransactionObject))
 		{
 			if (ObjectUpdate.ObjectData.bIsPendingKill)
@@ -399,10 +405,6 @@ void ProcessTransactionEvent(const FConcertTransactionEventBase& InEvent, const
 					DeletedActors->Add(TransactionActor);
 				}
 			}
-			else if (TransactionObjectRef.NeedsPostSpawn())
-			{
-				TransactionActor->FinishSpawning(FTransform(), true);
-			}
 		}
 
 #if WITH_EDITOR
@@ -524,17 +526,17 @@ FOnApplyTransaction& FConcertClientTransactionBridge::OnApplyTransaction()
 
 void FConcertClientTransactionBridge::ApplyRemoteTransaction(const FConcertTransactionEventBase& InEvent, const FConcertSessionVersionInfo* InVersionInfo, const TArray<FName>& InPackagesToProcess, const FConcertLocalIdentifierTable* InLocalIdentifierTablePtr, const bool bIsSnapshot)
 {
-	ConcertClientTransactionBridgeUtil::ProcessTransactionEvent(InEvent, InVersionInfo, InPackagesToProcess, InLocalIdentifierTablePtr, bIsSnapshot, FConcertSyncWorldRemapper(), nullptr);
+	ConcertClientTransactionBridgeUtil::ProcessTransactionEvent(InEvent, InVersionInfo, InPackagesToProcess, InLocalIdentifierTablePtr, bIsSnapshot, FConcertSyncWorldRemapper(), nullptr, bIncludeEditorOnlyProperties);
 }
 
 void FConcertClientTransactionBridge::ApplyRemoteTransaction(const FConcertTransactionEventBase& InEvent, const FConcertSessionVersionInfo* InVersionInfo, const TArray<FName>& InPackagesToProcess, const FConcertLocalIdentifierTable* InLocalIdentifierTablePtr, const bool bIsSnapshot, const FConcertSyncWorldRemapper& ConcertSyncWorldRemapper)
 {
-	ConcertClientTransactionBridgeUtil::ProcessTransactionEvent(InEvent, InVersionInfo, InPackagesToProcess, InLocalIdentifierTablePtr, bIsSnapshot, ConcertSyncWorldRemapper, nullptr);
+	ConcertClientTransactionBridgeUtil::ProcessTransactionEvent(InEvent, InVersionInfo, InPackagesToProcess, InLocalIdentifierTablePtr, bIsSnapshot, ConcertSyncWorldRemapper, nullptr, bIncludeEditorOnlyProperties);
 }
 
 void FConcertClientTransactionBridge::ApplyRemoteTransaction(const FConcertTransactionEventBase& InEvent, const FConcertSessionVersionInfo* InVersionInfo, const TArray<FName>& InPackagesToProcess, const FConcertLocalIdentifierTable* InLocalIdentifierTablePtr, const bool bIsSnapshot, const FConcertSyncWorldRemapper& ConcertSyncWorldRemapper, TSet<AActor*>* DeletedActors)
 {
-	ConcertClientTransactionBridgeUtil::ProcessTransactionEvent(InEvent, InVersionInfo, InPackagesToProcess, InLocalIdentifierTablePtr, bIsSnapshot, ConcertSyncWorldRemapper, DeletedActors);
+	ConcertClientTransactionBridgeUtil::ProcessTransactionEvent(InEvent, InVersionInfo, InPackagesToProcess, InLocalIdentifierTablePtr, bIsSnapshot, ConcertSyncWorldRemapper, DeletedActors, bIncludeEditorOnlyProperties);
 }
 
 bool& FConcertClientTransactionBridge::GetIgnoreLocalTransactionsRef()
diff --git a/Engine/Plugins/Developer/Concert/ConcertSync/ConcertSyncClient/Source/ConcertSyncClient/Private/ConcertSyncClientUtil.cpp b/Engine/Plugins/Developer/Concert/ConcertSync/ConcertSyncClient/Source/ConcertSyncClient/Private/ConcertSyncClientUtil.cpp
--- a/Engine/Plugins/Developer/Concert/ConcertSync/ConcertSyncClient/Source/ConcertSyncClient/Private/ConcertSyncClientUtil.cpp
+++ b/Engine/Plugins/Developer/Concert/ConcertSync/ConcertSyncClient/Source/ConcertSyncClient/Private/ConcertSyncClientUtil.cpp
@@ -282,9 +282,8 @@ FGetObjectResult GetObject(const FConcertObjectId& InObjectId, const FName InNew
 							SpawnParams.OverrideLevel = OuterLevel;
 							SpawnParams.SpawnCollisionHandlingOverride = ESpawnActorCollisionHandlingMethod::AlwaysSpawn;
 							SpawnParams.bNoFail = true;
-							SpawnParams.bDeferConstruction = true; // We defer FinishSpawning until the correct object state has been applied
 							SpawnParams.ObjectFlags = (EObjectFlags)InObjectId.ObjectPersistentFlags;
-							ObjectResult = FGetObjectResult(OwnerWorld->SpawnActor<AActor>(ObjectClass, FVector::ZeroVector, FRotator::ZeroRotator, SpawnParams), EGetObjectResultFlags::NeedsPostSpawn);
+							ObjectResult = FGetObjectResult(OwnerWorld->SpawnActor<AActor>(ObjectClass, FVector::ZeroVector, FRotator::ZeroRotator, SpawnParams), EGetObjectResultFlags::NewlyCreated);
 						}
 						else
 						{
@@ -298,7 +297,7 @@ FGetObjectResult GetObject(const FConcertObjectId& InObjectId, const FName InNew
 				}
 				else
 				{
-					ObjectResult = FGetObjectResult(NewObject<UObject>(NewObjectOuter, ObjectClass, *ObjectNameToCreate.ToString(), (EObjectFlags)InObjectId.ObjectPersistentFlags));
+					ObjectResult = FGetObjectResult(NewObject<UObject>(NewObjectOuter, ObjectClass, *ObjectNameToCreate.ToString(), (EObjectFlags)InObjectId.ObjectPersistentFlags), EGetObjectResultFlags::NewlyCreated);
 				}
 				
 				// if we have any package assignment, do it here
diff --git a/Engine/Plugins/Developer/Concert/ConcertSync/ConcertSyncClient/Source/ConcertSyncClient/Private/ConcertSyncClientUtil.h b/Engine/Plugins/Developer/Concert/ConcertSync/ConcertSyncClient/Source/ConcertSyncClient/Private/ConcertSyncClientUtil.h
--- a/Engine/Plugins/Developer/Concert/ConcertSync/ConcertSyncClient/Source/ConcertSyncClient/Private/ConcertSyncClientUtil.h
+++ b/Engine/Plugins/Developer/Concert/ConcertSync/ConcertSyncClient/Source/ConcertSyncClient/Private/ConcertSyncClientUtil.h
@@ -34,8 +34,8 @@ namespace ConcertSyncClientUtil
 	enum class EGetObjectResultFlags : uint8
 	{
 		None = 0,
-		NeedsGC = 1 << 0,
-		NeedsPostSpawn = 1 << 1,
+		NeedsGC = 1<<0,
+		NewlyCreated = 1<<1,
 	};
 	ENUM_CLASS_FLAGS(EGetObjectResultFlags);
 
@@ -58,9 +58,9 @@ namespace ConcertSyncClientUtil
 			return EnumHasAnyFlags(Flags, EGetObjectResultFlags::NeedsGC);
 		}
 
-		bool NeedsPostSpawn() const
+		bool NewlyCreated() const
 		{
-			return EnumHasAnyFlags(Flags, EGetObjectResultFlags::NeedsPostSpawn);
+			return EnumHasAnyFlags(Flags, EGetObjectResultFlags::NewlyCreated);
 		}
 
 		UObject* Obj;
