When creating a new object via Concert, only send the data that differs from the archetype

The destination object should match its archetype state when newly created (or have been reset to its archetype state when bResetExisting), so we can skip sending any non-overridden data.

diff --git a/Engine/Plugins/Developer/Concert/ConcertSync/ConcertSyncClient/Source/ConcertSyncClient/Private/ConcertClientTransactionBridge.cpp b/Engine/Plugins/Developer/Concert/ConcertSync/ConcertSyncClient/Source/ConcertSyncClient/Private/ConcertClientTransactionBridge.cpp
--- a/Engine/Plugins/Developer/Concert/ConcertSync/ConcertSyncClient/Source/ConcertSyncClient/Private/ConcertClientTransactionBridge.cpp
+++ b/Engine/Plugins/Developer/Concert/ConcertSync/ConcertSyncClient/Source/ConcertSyncClient/Private/ConcertClientTransactionBridge.cpp
@@ -7,6 +7,7 @@
 
 #include "ConcertTransactionEvents.h"
 #include "IConcertClientTransactionBridge.h"
+#include "TransactionCommon.h"
 #include "Misc/ITransactionObjectAnnotation.h"
 #include "Misc/PackageName.h"
 #include "Misc/CoreDelegates.h"
@@ -711,7 +712,6 @@ void FConcertClientTransactionBridge::HandleObjectTransacted(UObject* InObject,
 				ObjectUpdatePtr = &OngoingTransaction.SnapshotData.SnapshotObjectUpdates.AddDefaulted_GetRef();
 				ObjectUpdatePtr->ObjectId = ObjectId;
 				ObjectUpdatePtr->ObjectPathDepth = ConcertSyncClientUtil::GetObjectPathDepth(InObject);
-				ObjectUpdatePtr->ObjectData.bAllowCreate = false;
 				ObjectUpdatePtr->ObjectData.bIsPendingKill = !IsValid(InObject);
 			}
 
@@ -742,10 +742,23 @@ void FConcertClientTransactionBridge::HandleObjectTransacted(UObject* InObject,
 	}
 	else if (OnLocalTransactionFinalizedDelegate.IsBound())
 	{
+		auto CanObjectBeCreated = [](const UObject* ObjToTest) -> bool
+		{
+			if (const UActorComponent* ComponentToTest = Cast<UActorComponent>(ObjToTest))
+			{
+				// Components that are managed by a native or Blueprint class cannot be created, so we only allow "instance" components to be created
+				return ComponentToTest->CreationMethod == EComponentCreationMethod::Instance;
+			}
+			return true;
+		};
+
+		const bool bIsNewlyCreated = InTransactionEvent.HasPendingKillChange() && IsValid(InObject);
+
 		FConcertExportedObject& ObjectUpdate = OngoingTransaction.FinalizedData.FinalizedObjectUpdates.AddDefaulted_GetRef();
 		ObjectUpdate.ObjectId = ObjectId;
 		ObjectUpdate.ObjectPathDepth = ConcertSyncClientUtil::GetObjectPathDepth(InObject);
-		ObjectUpdate.ObjectData.bAllowCreate = InTransactionEvent.HasPendingKillChange() && IsValid(InObject);
+		ObjectUpdate.ObjectData.bAllowCreate = bIsNewlyCreated && CanObjectBeCreated(InObject);
+		ObjectUpdate.ObjectData.bResetExisting = bIsNewlyCreated;
 		ObjectUpdate.ObjectData.bIsPendingKill = !IsValid(InObject);
 		ObjectUpdate.ObjectData.NewPackageName = NewObjectPackageName;
 		ObjectUpdate.ObjectData.NewName = NewObjectName;
@@ -758,31 +771,57 @@ void FConcertClientTransactionBridge::HandleObjectTransacted(UObject* InObject,
 			TransactionAnnotation->Serialize(AnnotationWriter);
 		}
 
-		// If this object changed from being pending kill to not being pending kill, we have to send a full object update (including all properties), rather than attempt a delta-update
-		const bool bForceFullObjectUpdate = InTransactionEvent.HasPendingKillChange() && IsValid(InObject);
-		if (bForceFullObjectUpdate)
+		auto PopulateObjectUpdateData = [this, &OngoingTransaction, &ObjectUpdate, InObject](const bool bHasNonPropertyChanges, const TArray<const FProperty*>* ExportedPropertiesPtr)
 		{
-			// Serialize the entire object.
-			ConcertSyncClientUtil::SerializeObject(&OngoingTransaction.FinalizedData.FinalizedLocalIdentifierTable, InObject, nullptr, bIncludeEditorOnlyProperties, ObjectUpdate.ObjectData.SerializedData);
-		}
-		else if (InTransactionEvent.HasNonPropertyChanges(/*SerializationOnly*/true))
+			if (bHasNonPropertyChanges)
+			{
+				// The 'non-property changes' refers to custom data added by a deriver UObject before and/or after the standard serialized data. Since this is a custom
+				// data format, we don't know what changed, call the object to re-serialize this part, but still send the delta for the generic reflected properties (in RootPropertyNames).
+				ConcertSyncClientUtil::SerializeObject(&OngoingTransaction.FinalizedData.FinalizedLocalIdentifierTable, InObject, ExportedPropertiesPtr, bIncludeEditorOnlyProperties, ObjectUpdate.ObjectData.SerializedData);
+
+				if (ExportedPropertiesPtr)
+				{
+					// Track which properties changed. Not used to apply the transaction on the receiving side, the object specific serialization function will be used for that, but
+					// to be able to display, in the transaction detail view, which 'properties' changed in the transaction as transaction data is otherwise opaque to UI.
+					for (const FProperty* ExportedProperty : *ExportedPropertiesPtr)
+					{
+						FConcertSerializedPropertyData& PropertyData = ObjectUpdate.PropertyDatas.AddDefaulted_GetRef();
+						PropertyData.PropertyName = ExportedProperty->GetFName();
+					}
+				}
+			}
+			else if (ExportedPropertiesPtr) // Its possible to optimize the transaction payload, only sending a 'delta' update.
+			{
+				// Only send properties that changed. The receiving side will 'patch' the object using the reflection system. The specific object serialization function will NOT be called.
+				ConcertSyncClientUtil::SerializeProperties(&OngoingTransaction.FinalizedData.FinalizedLocalIdentifierTable, InObject, *ExportedPropertiesPtr, bIncludeEditorOnlyProperties, ObjectUpdate.PropertyDatas);
+			}
+		};
+
+		// If this object changed from being pending kill to not being pending kill then we'll send an update against the archetype state (rather than use 
+		// the current delta-update) and the receiving side will ensure that the target object matches the archetype state prior to applying the update
+		if (bIsNewlyCreated)
 		{
-			// The 'non-property changes' refers to custom data added by a deriver UObject before and/or after the standard serialized data. Since this is a custom
-			// data format, we don't know what changed, call the object to re-serialize this part, but still send the delta for the generic reflected properties (in RootPropertyNames).
-			ConcertSyncClientUtil::SerializeObject(&OngoingTransaction.FinalizedData.FinalizedLocalIdentifierTable, InObject, &ExportedProperties, bIncludeEditorOnlyProperties, ObjectUpdate.ObjectData.SerializedData);
+			UObject* ObjectArchetype = InObject->GetArchetype();
+			if (ObjectArchetype && !InObject->HasAnyFlags(RF_ClassDefaultObject | RF_ArchetypeObject))
+			{
+				UE::Transaction::FDiffableObject ArchetypeDiffableObject = UE::Transaction::DiffUtil::GetDiffableObject(ObjectArchetype);
+				ArchetypeDiffableObject.SetObject(InObject); // Use the real object here, as the archetype info will never be valid for an instance
 
-			// Track which properties changed. Not used to apply the transaction on the receiving side, the object specific serialization function will be used for that, but
-			// to be able to display, in the transaction detail view, which 'properties' changed in the transaction as transaction data is otherwise opaque to UI.
-			for (const FProperty* ExportedProperty : ExportedProperties)
+				const UE::Transaction::FDiffableObject CurrentDiffableObject = UE::Transaction::DiffUtil::GetDiffableObject(InObject);
+
+				const FTransactionObjectDeltaChange DeltaChangeFromArchetype = UE::Transaction::DiffUtil::GenerateObjectDiff(ArchetypeDiffableObject, CurrentDiffableObject);
+				const TArray<const FProperty*> ChangedPropertiesFromArchetype = ConcertSyncClientUtil::GetExportedProperties(InObject->GetClass(), DeltaChangeFromArchetype.ChangedProperties, bIncludeEditorOnlyProperties);
+
+				PopulateObjectUpdateData(DeltaChangeFromArchetype.bHasNonPropertyChanges, &ChangedPropertiesFromArchetype);
+			}
+			else
 			{
-				FConcertSerializedPropertyData& PropertyData = ObjectUpdate.PropertyDatas.AddDefaulted_GetRef();
-				PropertyData.PropertyName = ExportedProperty->GetFName();
+				PopulateObjectUpdateData(/*bHasNonPropertyChanges*/true, nullptr);
 			}
 		}
-		else // Its possible to optimize the transaction payload, only sending a 'delta' update.
+		else
 		{
-			// Only send properties that changed. The receiving side will 'patch' the object using the reflection system. The specific object serialization function will NOT be called.
-			ConcertSyncClientUtil::SerializeProperties(&OngoingTransaction.FinalizedData.FinalizedLocalIdentifierTable, InObject, ExportedProperties, bIncludeEditorOnlyProperties, ObjectUpdate.PropertyDatas);
+			PopulateObjectUpdateData(InTransactionEvent.HasNonPropertyChanges(/*SerializationOnly*/true), &ExportedProperties);
 		}
 	}
 }
