Preparing to remove a private include hack

#rnx

diff --git a/Engine/Plugins/Developer/Concert/ConcertSync/ConcertSyncClient/Source/ConcertSyncClient/Private/ConcertClientTransactionBridge.cpp b/Engine/Plugins/Developer/Concert/ConcertSync/ConcertSyncClient/Source/ConcertSyncClient/Private/ConcertClientTransactionBridge.cpp
--- a/Engine/Plugins/Developer/Concert/ConcertSync/ConcertSyncClient/Source/ConcertSyncClient/Private/ConcertClientTransactionBridge.cpp
+++ b/Engine/Plugins/Developer/Concert/ConcertSync/ConcertSyncClient/Source/ConcertSyncClient/Private/ConcertClientTransactionBridge.cpp
@@ -470,6 +470,11 @@ void ProcessTransactionEvent(const FConcertTransactionEventBase& InEvent, const
 
 }	// namespace ConcertClientTransactionBridgeUtil
 
+TUniquePtr<IConcertClientTransactionBridge> IConcertClientTransactionBridge::NewInstance()
+{
+	return MakeUnique<FConcertClientTransactionBridge>();
+}
+
 FConcertClientTransactionBridge::FConcertClientTransactionBridge()
 	: bHasBoundUnderlyingLocalTransactionEvents(false)
 	, bIgnoreLocalTransactions(false)
@@ -505,6 +510,11 @@ FConcertClientTransactionBridge::~FConcertClientTransactionBridge()
 	FCoreDelegates::OnEndFrame.RemoveAll(this);
 }
 
+void FConcertClientTransactionBridge::SetIncludeEditorOnlyProperties(const bool InIncludeEditorOnlyProperties)
+{
+	bIncludeEditorOnlyProperties = InIncludeEditorOnlyProperties;
+}
+
 FOnConcertClientLocalTransactionSnapshot& FConcertClientTransactionBridge::OnLocalTransactionSnapshot()
 {
 	return OnLocalTransactionSnapshotDelegate;
diff --git a/Engine/Plugins/Developer/Concert/ConcertSync/ConcertSyncClient/Source/ConcertSyncClient/Private/ConcertClientTransactionBridge.h b/Engine/Plugins/Developer/Concert/ConcertSync/ConcertSyncClient/Source/ConcertSyncClient/Private/ConcertClientTransactionBridge.h
--- a/Engine/Plugins/Developer/Concert/ConcertSync/ConcertSyncClient/Source/ConcertSyncClient/Private/ConcertClientTransactionBridge.h
+++ b/Engine/Plugins/Developer/Concert/ConcertSync/ConcertSyncClient/Source/ConcertSyncClient/Private/ConcertClientTransactionBridge.h
@@ -14,6 +14,7 @@ public:
 	virtual ~FConcertClientTransactionBridge();
 
 	//~ IConcertClientTransactionBridge interface
+	virtual void SetIncludeEditorOnlyProperties(const bool InIncludeEditorOnlyProperties) override;
 	virtual FOnConcertClientLocalTransactionSnapshot& OnLocalTransactionSnapshot() override;
 	virtual FOnConcertClientLocalTransactionFinalized& OnLocalTransactionFinalized() override;
 	virtual bool CanApplyRemoteTransaction() const override;
diff --git a/Engine/Plugins/Developer/Concert/ConcertSync/ConcertSyncClient/Source/ConcertSyncClient/Public/IConcertClientTransactionBridge.h b/Engine/Plugins/Developer/Concert/ConcertSync/ConcertSyncClient/Source/ConcertSyncClient/Public/IConcertClientTransactionBridge.h
--- a/Engine/Plugins/Developer/Concert/ConcertSync/ConcertSyncClient/Source/ConcertSyncClient/Public/IConcertClientTransactionBridge.h
+++ b/Engine/Plugins/Developer/Concert/ConcertSync/ConcertSyncClient/Source/ConcertSyncClient/Public/IConcertClientTransactionBridge.h
@@ -71,6 +71,11 @@ DECLARE_MULTICAST_DELEGATE_TwoParams(FOnConcertClientLocalTransactionFinalized,
 class IConcertClientTransactionBridge
 {
 public:
+	/**
+	 * Create a new instance of the concrete implementation of this class. 
+	 */
+	CONCERTSYNCCLIENT_API static TUniquePtr<IConcertClientTransactionBridge> NewInstance();
+
 	/** Scoped struct to ignore a local transaction */
 	struct FScopedIgnoreLocalTransaction : private TGuardValue<bool>
 	{
@@ -82,6 +87,12 @@ public:
 
 	virtual ~IConcertClientTransactionBridge() = default;
 
+	/**
+	 * Set whether or not to include editor-only properties when serializing object and property changes.
+	 * @note This is set to true by default.
+	 */
+	virtual void SetIncludeEditorOnlyProperties(const bool InIncludeEditorOnlyProperties) = 0;
+
 	/**
 	 * Called when an ongoing transaction is updated via a snapshot.
 	 * @note This is called during end-frame processing.
