Removed a private include hack

#rnx

diff --git a/Engine/Plugins/Developer/Concert/ConcertSync/ConcertSyncClient/Source/ConcertSyncClient/Private/ConcertClientTransactionBridge.cpp b/Engine/Plugins/Developer/Concert/ConcertSync/ConcertSyncClient/Source/ConcertSyncClient/Private/ConcertClientTransactionBridge.cpp
--- a/Engine/Plugins/Developer/Concert/ConcertSync/ConcertSyncClient/Source/ConcertSyncClient/Private/ConcertClientTransactionBridge.cpp
+++ b/Engine/Plugins/Developer/Concert/ConcertSync/ConcertSyncClient/Source/ConcertSyncClient/Private/ConcertClientTransactionBridge.cpp
@@ -486,12 +486,6 @@ FConcertClientTransactionBridge::FConcertClientTransactionBridge()
 	FCoreDelegates::OnEndFrame.AddRaw(this, &FConcertClientTransactionBridge::OnEndFrame);
 }
 
-FConcertClientTransactionBridge::FConcertClientTransactionBridge(bool bInIncludeEditorOnlyProperties)
-	: FConcertClientTransactionBridge()
-{
-	bIncludeEditorOnlyProperties = bInIncludeEditorOnlyProperties;
-}
-
 FConcertClientTransactionBridge::~FConcertClientTransactionBridge()
 {
 #if WITH_EDITOR
diff --git a/Engine/Plugins/Developer/Concert/ConcertSync/ConcertSyncClient/Source/ConcertSyncClient/Private/ConcertClientTransactionBridge.h b/Engine/Plugins/Developer/Concert/ConcertSync/ConcertSyncClient/Source/ConcertSyncClient/Private/ConcertClientTransactionBridge.h
--- a/Engine/Plugins/Developer/Concert/ConcertSync/ConcertSyncClient/Source/ConcertSyncClient/Private/ConcertClientTransactionBridge.h
+++ b/Engine/Plugins/Developer/Concert/ConcertSync/ConcertSyncClient/Source/ConcertSyncClient/Private/ConcertClientTransactionBridge.h
@@ -6,11 +6,10 @@
 #include "Misc/ITransaction.h"
 #include "ConcertSyncArchives.h"
 
-class CONCERTSYNCCLIENT_API FConcertClientTransactionBridge : public IConcertClientTransactionBridge
+class FConcertClientTransactionBridge : public IConcertClientTransactionBridge
 {
 public:
 	FConcertClientTransactionBridge();
-	FConcertClientTransactionBridge(bool bInIncludeEditorOnlyProperties);
 	virtual ~FConcertClientTransactionBridge();
 
 	//~ IConcertClientTransactionBridge interface
