Improved Concert detection of newly created objects

Inner objects of a newly created object are now also tracked as newly created so that they're not reset incorrectly, eg) components within a newly created actor

#rnx

diff --git a/Engine/Plugins/Developer/Concert/ConcertSync/ConcertSyncClient/Source/ConcertSyncClient/Private/ConcertClientTransactionBridge.cpp b/Engine/Plugins/Developer/Concert/ConcertSync/ConcertSyncClient/Source/ConcertSyncClient/Private/ConcertClientTransactionBridge.cpp
--- a/Engine/Plugins/Developer/Concert/ConcertSync/ConcertSyncClient/Source/ConcertSyncClient/Private/ConcertClientTransactionBridge.cpp
+++ b/Engine/Plugins/Developer/Concert/ConcertSync/ConcertSyncClient/Source/ConcertSyncClient/Private/ConcertClientTransactionBridge.cpp
@@ -243,6 +243,8 @@ void ProcessTransactionEvent(const FConcertTransactionEventBase& InEvent, const
 	bool bObjectsDeleted = false;
 	TArray<ConcertSyncClientUtil::FGetObjectResult, TInlineAllocator<8>> TransactionObjects;
 	{
+		TSet<const UObject*> NewlyCreatedObjects;
+
 		// Find or create each object in parent->child order
 		TransactionObjects.AddDefaulted(SortedExportedObjects.Num());
 		for (int32 ObjectIndex = 0; ObjectIndex < SortedExportedObjects.Num(); ++ObjectIndex)
@@ -265,9 +267,20 @@ void ProcessTransactionEvent(const FConcertTransactionEventBase& InEvent, const
 			TransactionObjectRef = ConcertSyncClientUtil::GetObject(ObjectUpdate.ObjectId, ObjectUpdate.ObjectData.NewName, ObjectUpdate.ObjectData.NewOuterPathName, ObjectUpdate.ObjectData.NewExternalPackageName, ObjectUpdate.ObjectData.bAllowCreate);
 			bObjectsDeleted |= (ObjectUpdate.ObjectData.bIsPendingKill || TransactionObjectRef.NeedsGC());
 
-			if (ObjectUpdate.ObjectData.bResetExisting && !TransactionObjectRef.NewlyCreated())
+			if (TransactionObjectRef.NewlyCreated())
 			{
 				check(TransactionObjectRef.Obj);
+
+				// Track this object (and any inner objects, as they must also be new) as newly created
+				NewlyCreatedObjects.Add(TransactionObjectRef.Obj);
+				ForEachObjectWithOuter(TransactionObjectRef.Obj, [&NewlyCreatedObjects](UObject* InnerObj)
+				{
+					NewlyCreatedObjects.Add(InnerObj);
+				});
+			}
+
+			if (ObjectUpdate.ObjectData.bResetExisting && !NewlyCreatedObjects.Contains(TransactionObjectRef.Obj))
+			{
 				ConcertSyncUtil::ResetObjectPropertiesToArchetypeValues(TransactionObjectRef.Obj, bIncludeEditorOnlyProperties);
 			}
 		}
