Added a null check for newly created Concert objects

#rnx

diff --git a/Engine/Plugins/Developer/Concert/ConcertSync/ConcertSyncClient/Source/ConcertSyncClient/Private/ConcertClientTransactionBridge.cpp b/Engine/Plugins/Developer/Concert/ConcertSync/ConcertSyncClient/Source/ConcertSyncClient/Private/ConcertClientTransactionBridge.cpp
--- a/Engine/Plugins/Developer/Concert/ConcertSync/ConcertSyncClient/Source/ConcertSyncClient/Private/ConcertClientTransactionBridge.cpp
+++ b/Engine/Plugins/Developer/Concert/ConcertSync/ConcertSyncClient/Source/ConcertSyncClient/Private/ConcertClientTransactionBridge.cpp
@@ -267,10 +267,8 @@ void ProcessTransactionEvent(const FConcertTransactionEventBase& InEvent, const
 			TransactionObjectRef = ConcertSyncClientUtil::GetObject(ObjectUpdate.ObjectId, ObjectUpdate.ObjectData.NewName, ObjectUpdate.ObjectData.NewOuterPathName, ObjectUpdate.ObjectData.NewExternalPackageName, ObjectUpdate.ObjectData.bAllowCreate);
 			bObjectsDeleted |= (ObjectUpdate.ObjectData.bIsPendingKill || TransactionObjectRef.NeedsGC());
 
-			if (TransactionObjectRef.NewlyCreated())
+			if (TransactionObjectRef.Obj && TransactionObjectRef.NewlyCreated())
 			{
-				check(TransactionObjectRef.Obj);
-
 				// Track this object (and any inner objects, as they must also be new) as newly created
 				NewlyCreatedObjects.Add(TransactionObjectRef.Obj);
 				ForEachObjectWithOuter(TransactionObjectRef.Obj, [&NewlyCreatedObjects](UObject* InnerObj)
