Allow Concert to optionally send object changes generated from a transaction annotation, rather than send the transaction annotation itself

This can be controlled by UConcertSyncConfig::bIncludeAnnotationObjectChanges and is off by default. This feature should be considered experimental, but can result in sending significantly less data for certain types of transaction annotation (such as FActorTransactionAnnotation, which serializes entire components).

diff --git a/Engine/Plugins/Developer/Concert/ConcertSync/ConcertSyncClient/Source/ConcertSyncClient/Private/ConcertClientTransactionBridge.cpp b/Engine/Plugins/Developer/Concert/ConcertSync/ConcertSyncClient/Source/ConcertSyncClient/Private/ConcertClientTransactionBridge.cpp
--- a/Engine/Plugins/Developer/Concert/ConcertSync/ConcertSyncClient/Source/ConcertSyncClient/Private/ConcertClientTransactionBridge.cpp
+++ b/Engine/Plugins/Developer/Concert/ConcertSync/ConcertSyncClient/Source/ConcertSyncClient/Private/ConcertClientTransactionBridge.cpp
@@ -223,7 +223,7 @@ void ProcessTransactionEvent(const FConcertTransactionEventBase& InEvent, const
 	// --------------------------------------------------------------------------------------------------------------------
 	// Phase 0
 	// --------------------------------------------------------------------------------------------------------------------
-	TArray<const FConcertExportedObject*, TInlineAllocator<8>> SortedExportedObjects;
+	TArray<const FConcertExportedObject*, TInlineAllocator<32>> SortedExportedObjects;
 	{
 		SortedExportedObjects.Reserve(InEvent.ExportedObjects.Num());
 		for (const FConcertExportedObject& ExportedObject : InEvent.ExportedObjects)
@@ -241,7 +241,7 @@ void ProcessTransactionEvent(const FConcertTransactionEventBase& InEvent, const
 	// Phase 1
 	// --------------------------------------------------------------------------------------------------------------------
 	bool bObjectsDeleted = false;
-	TArray<ConcertSyncClientUtil::FGetObjectResult, TInlineAllocator<8>> TransactionObjects;
+	TArray<ConcertSyncClientUtil::FGetObjectResult, TInlineAllocator<32>> TransactionObjects;
 	{
 		TSet<const UObject*> NewlyCreatedObjects;
 
@@ -297,7 +297,7 @@ void ProcessTransactionEvent(const FConcertTransactionEventBase& InEvent, const
 	// Phase 2
 	// --------------------------------------------------------------------------------------------------------------------
 #if WITH_EDITOR
-	TArray<TSharedPtr<ITransactionObjectAnnotation>, TInlineAllocator<8>> TransactionAnnotations;
+	TArray<TSharedPtr<ITransactionObjectAnnotation>, TInlineAllocator<32>> TransactionAnnotations;
 	TransactionAnnotations.AddDefaulted(SortedExportedObjects.Num());
 	for (int32 ObjectIndex = SortedExportedObjects.Num() - 1; ObjectIndex >= 0; --ObjectIndex)
 	{
@@ -490,6 +490,7 @@ FConcertClientTransactionBridge::FConcertClientTransactionBridge()
 	: bHasBoundUnderlyingLocalTransactionEvents(false)
 	, bIgnoreLocalTransactions(false)
 	, bIncludeEditorOnlyProperties(true)
+	, bIncludeAnnotationObjectChanges(GetDefault<UConcertSyncConfig>()->bIncludeAnnotationObjectChanges)
 {
 	ConditionalBindUnderlyingLocalTransactionEvents();
 
@@ -520,6 +521,11 @@ void FConcertClientTransactionBridge::SetIncludeEditorOnlyProperties(const bool
 	bIncludeEditorOnlyProperties = InIncludeEditorOnlyProperties;
 }
 
+void FConcertClientTransactionBridge::SetIncludeAnnotationObjectChanges(const bool InIncludeAnnotationObjectChanges)
+{
+	bIncludeAnnotationObjectChanges = InIncludeAnnotationObjectChanges;
+}
+
 FOnConcertClientLocalTransactionSnapshot& FConcertClientTransactionBridge::OnLocalTransactionSnapshot()
 {
 	return OnLocalTransactionSnapshotDelegate;
@@ -634,6 +640,11 @@ void FConcertClientTransactionBridge::HandleObjectTransacted(UObject* InObject,
 		return;
 	}
 
+	if (!bIncludeAnnotationObjectChanges && InTransactionEvent.GetObjectChangeCreatedBy() == ETransactionObjectChangeCreatedBy::TransactionAnnotation)
+	{
+		return;
+	}
+
 	UPackage* ChangedPackage = InObject->GetOutermost();
 	ETransactionFilterResult FilterResult = ConcertClientTransactionBridgeUtil::ApplyTransactionFilters(TransactionFilters, InObject, ChangedPackage);
 	FOngoingTransaction* TrackedTransaction = OngoingTransactions.Find(InTransactionEvent.GetOperationId());
@@ -676,7 +687,7 @@ void FConcertClientTransactionBridge::HandleObjectTransacted(UObject* InObject,
 		return;
 	}
 
-	const FConcertObjectId ObjectId = FConcertObjectId(*InObject->GetClass()->GetPathName(), InTransactionEvent.GetOriginalObjectPackageName(), InTransactionEvent.GetOriginalObjectName(), InTransactionEvent.GetOriginalObjectOuterPathName(), InTransactionEvent.GetOriginalObjectExternalPackageName(), InObject->GetFlags());
+	const FConcertObjectId ObjectId = FConcertObjectId(InTransactionEvent.GetOriginalObjectId(), InObject->GetFlags());
 	FOngoingTransaction& OngoingTransaction = *TrackedTransaction;
 
 	// If the object is excluded or exclude the whole transaction add it to the excluded list
@@ -687,12 +698,26 @@ void FConcertClientTransactionBridge::HandleObjectTransacted(UObject* InObject,
 		return;
 	}
 
-	const FName NewObjectPackageName = InTransactionEvent.HasOuterChange() || InTransactionEvent.HasExternalPackageChange() ? InObject->GetPackage()->GetFName() : FName();
-	const FName NewObjectName = InTransactionEvent.HasNameChange() ? InObject->GetFName() : FName();
-	const FName NewObjectOuterPathName = (InTransactionEvent.HasOuterChange() && InObject->GetOuter()) ? FName(*InObject->GetOuter()->GetPathName()) : FName();
-	const FName NewObjectExternalPackageName = (InTransactionEvent.HasExternalPackageChange() && InObject->GetExternalPackage()) ? InObject->GetExternalPackage()->GetFName() : FName();
+	bool bCanCreateOrRenameObject = true;
+	if (const UActorComponent* Component = Cast<UActorComponent>(InObject))
+	{
+		bool bIsOwnedByChildActor = false;
+		if (const AActor* ComponentOwner = Component->GetOwner())
+		{
+			bIsOwnedByChildActor = ComponentOwner->IsChildActor();
+		}
+
+		// Components that are managed by a native or Blueprint class cannot be created or renamed, so we only allow "instance" components to be created or renamed
+		bCanCreateOrRenameObject = !bIsOwnedByChildActor && Component->CreationMethod == EComponentCreationMethod::Instance;
+	}
+
+	const FName NewObjectPackageName = (bCanCreateOrRenameObject && (InTransactionEvent.HasOuterChange() || InTransactionEvent.HasExternalPackageChange())) ? InObject->GetPackage()->GetFName() : FName();
+	const FName NewObjectName = (bCanCreateOrRenameObject && InTransactionEvent.HasNameChange()) ? InObject->GetFName() : FName();
+	const FName NewObjectOuterPathName = (bCanCreateOrRenameObject && (InTransactionEvent.HasOuterChange() && InObject->GetOuter())) ? FName(*InObject->GetOuter()->GetPathName()) : FName();
+	const FName NewObjectExternalPackageName = (bCanCreateOrRenameObject && (InTransactionEvent.HasExternalPackageChange() && InObject->GetExternalPackage())) ? InObject->GetExternalPackage()->GetFName() : FName();
 	const TArray<const FProperty*> ExportedProperties = ConcertSyncClientUtil::GetExportedProperties(InObject->GetClass(), InTransactionEvent.GetChangedProperties(), bIncludeEditorOnlyProperties);
 	TSharedPtr<ITransactionObjectAnnotation> TransactionAnnotation = InTransactionEvent.GetAnnotation();
+	const bool bUseSerializedAnnotationData = !bIncludeAnnotationObjectChanges || !TransactionAnnotation || !TransactionAnnotation->SupportsAdditionalObjectChanges();
 
 	// Track which packages were changed
 	OngoingTransaction.CommonData.ModifiedPackages.AddUnique(ChangedPackage->GetFName());
@@ -730,7 +755,7 @@ void FConcertClientTransactionBridge::HandleObjectTransacted(UObject* InObject,
 				ObjectUpdatePtr->ObjectData.bIsPendingKill = !IsValid(InObject);
 			}
 
-			if (TransactionAnnotation.IsValid())
+			if (TransactionAnnotation && bUseSerializedAnnotationData)
 			{
 				ObjectUpdatePtr->SerializedAnnotationData.Reset();
 				FConcertSyncObjectWriter AnnotationWriter(nullptr, InObject, ObjectUpdatePtr->SerializedAnnotationData, bIncludeEditorOnlyProperties, true);
@@ -757,22 +782,14 @@ void FConcertClientTransactionBridge::HandleObjectTransacted(UObject* InObject,
 	}
 	else if (OnLocalTransactionFinalizedDelegate.IsBound())
 	{
-		auto CanObjectBeCreated = [](const UObject* ObjToTest) -> bool
-		{
-			if (const UActorComponent* ComponentToTest = Cast<UActorComponent>(ObjToTest))
-			{
-				// Components that are managed by a native or Blueprint class cannot be created, so we only allow "instance" components to be created
-				return ComponentToTest->CreationMethod == EComponentCreationMethod::Instance;
-			}
-			return true;
-		};
+		
 
 		const bool bIsNewlyCreated = InTransactionEvent.HasPendingKillChange() && IsValid(InObject);
 
 		FConcertExportedObject& ObjectUpdate = OngoingTransaction.FinalizedData.FinalizedObjectUpdates.AddDefaulted_GetRef();
 		ObjectUpdate.ObjectId = ObjectId;
 		ObjectUpdate.ObjectPathDepth = ConcertSyncClientUtil::GetObjectPathDepth(InObject);
-		ObjectUpdate.ObjectData.bAllowCreate = bIsNewlyCreated && CanObjectBeCreated(InObject);
+		ObjectUpdate.ObjectData.bAllowCreate = bIsNewlyCreated && bCanCreateOrRenameObject;
 		ObjectUpdate.ObjectData.bResetExisting = bIsNewlyCreated;
 		ObjectUpdate.ObjectData.bIsPendingKill = !IsValid(InObject);
 		ObjectUpdate.ObjectData.NewPackageName = NewObjectPackageName;
@@ -780,7 +797,7 @@ void FConcertClientTransactionBridge::HandleObjectTransacted(UObject* InObject,
 		ObjectUpdate.ObjectData.NewOuterPathName = NewObjectOuterPathName;
 		ObjectUpdate.ObjectData.NewExternalPackageName = NewObjectExternalPackageName;
 
-		if (TransactionAnnotation.IsValid())
+		if (TransactionAnnotation && bUseSerializedAnnotationData)
 		{
 			FConcertSyncObjectWriter AnnotationWriter(&OngoingTransaction.FinalizedData.FinalizedLocalIdentifierTable, InObject, ObjectUpdate.SerializedAnnotationData, bIncludeEditorOnlyProperties, false);
 			TransactionAnnotation->Serialize(AnnotationWriter);
diff --git a/Engine/Plugins/Developer/Concert/ConcertSync/ConcertSyncClient/Source/ConcertSyncClient/Private/ConcertClientTransactionBridge.h b/Engine/Plugins/Developer/Concert/ConcertSync/ConcertSyncClient/Source/ConcertSyncClient/Private/ConcertClientTransactionBridge.h
--- a/Engine/Plugins/Developer/Concert/ConcertSync/ConcertSyncClient/Source/ConcertSyncClient/Private/ConcertClientTransactionBridge.h
+++ b/Engine/Plugins/Developer/Concert/ConcertSync/ConcertSyncClient/Source/ConcertSyncClient/Private/ConcertClientTransactionBridge.h
@@ -14,6 +14,7 @@ public:
 
 	//~ IConcertClientTransactionBridge interface
 	virtual void SetIncludeEditorOnlyProperties(const bool InIncludeEditorOnlyProperties) override;
+	virtual void SetIncludeAnnotationObjectChanges(const bool InIncludeAnnotationObjectChanges) override;
 	virtual FOnConcertClientLocalTransactionSnapshot& OnLocalTransactionSnapshot() override;
 	virtual FOnConcertClientLocalTransactionFinalized& OnLocalTransactionFinalized() override;
 	virtual bool CanApplyRemoteTransaction() const override;
@@ -83,5 +84,8 @@ private:
 	/** Include non-cooked properties in object serialization */
 	bool bIncludeEditorOnlyProperties;
 
+	/** Include object changes that have been generated via a transaction annotation (where possible) */
+	bool bIncludeAnnotationObjectChanges;
+
 	FConcertSyncWorldRemapper WorldRemapper = FConcertSyncWorldRemapper();
 };
diff --git a/Engine/Plugins/Developer/Concert/ConcertSync/ConcertSyncClient/Source/ConcertSyncClient/Public/IConcertClientTransactionBridge.h b/Engine/Plugins/Developer/Concert/ConcertSync/ConcertSyncClient/Source/ConcertSyncClient/Public/IConcertClientTransactionBridge.h
--- a/Engine/Plugins/Developer/Concert/ConcertSync/ConcertSyncClient/Source/ConcertSyncClient/Public/IConcertClientTransactionBridge.h
+++ b/Engine/Plugins/Developer/Concert/ConcertSync/ConcertSyncClient/Source/ConcertSyncClient/Public/IConcertClientTransactionBridge.h
@@ -93,6 +93,13 @@ public:
 	 */
 	virtual void SetIncludeEditorOnlyProperties(const bool InIncludeEditorOnlyProperties) = 0;
 
+	/**
+	 * Set whether to include object changes that have been generated via a transaction annotation 
+	 * (where possible), or whether to send the entire transaction annotation blob instead.
+	 * @note This is set to UConcertSyncConfig::bIncludeAnnotationObjectChanges by default.
+	 */
+	virtual void SetIncludeAnnotationObjectChanges(const bool InIncludeAnnotationObjectChanges) = 0;
+
 	/**
 	 * Called when an ongoing transaction is updated via a snapshot.
 	 * @note This is called during end-frame processing.
diff --git a/Engine/Plugins/Developer/Concert/ConcertSync/ConcertSyncCore/Source/ConcertSyncCore/Public/ConcertSyncSettings.h b/Engine/Plugins/Developer/Concert/ConcertSync/ConcertSyncCore/Source/ConcertSyncCore/Public/ConcertSyncSettings.h
--- a/Engine/Plugins/Developer/Concert/ConcertSync/ConcertSyncCore/Source/ConcertSyncCore/Public/ConcertSyncSettings.h
+++ b/Engine/Plugins/Developer/Concert/ConcertSync/ConcertSyncCore/Source/ConcertSyncCore/Public/ConcertSyncSettings.h
@@ -52,6 +52,7 @@ public:
 	UConcertSyncConfig()
 		: bInteractiveHotReload(false)
 		, bShowPresenceInPIE(true)
+		, bIncludeAnnotationObjectChanges(false)
 		, SnapshotTransactionsPerSecond(30.0f)
 	{}
 
@@ -68,6 +69,13 @@ public:
 	UPROPERTY(config, EditAnywhere, AdvancedDisplay, Category="Transaction Settings")
 	bool bShowPresenceInPIE;
 
+	/**
+	 * Should we include object changes that have been generated via a transaction annotation
+	 * (where possible), or should we send the entire transaction annotation blob instead.
+	 */
+	UPROPERTY(config, EditAnywhere, AdvancedDisplay, Category="Transaction Settings")
+	bool bIncludeAnnotationObjectChanges;
+
 	/**
 	 * Number of snapshot transactions (eg, moving an object or dragging a slider) that should be sent per-second to other clients.
 	 */
diff --git a/Engine/Plugins/Developer/Concert/ConcertSync/ConcertSyncCore/Source/ConcertSyncCore/Public/ConcertTransactionEvents.h b/Engine/Plugins/Developer/Concert/ConcertSync/ConcertSyncCore/Source/ConcertSyncCore/Public/ConcertTransactionEvents.h
--- a/Engine/Plugins/Developer/Concert/ConcertSync/ConcertSyncCore/Source/ConcertSyncCore/Public/ConcertTransactionEvents.h
+++ b/Engine/Plugins/Developer/Concert/ConcertSync/ConcertSyncCore/Source/ConcertSyncCore/Public/ConcertTransactionEvents.h
@@ -6,6 +6,7 @@
 #include "UObject/Object.h"
 #include "UObject/Class.h"
 #include "UObject/Package.h"
+#include "Misc/TransactionObjectEvent.h"
 #include "IdentifierTable/ConcertIdentifierTableData.h"
 #include "ConcertTransactionEvents.generated.h"
 
@@ -27,22 +28,25 @@ struct FConcertObjectId
 {
 	GENERATED_BODY()
 
-	FConcertObjectId()
-		: ObjectPersistentFlags(0)
+	FConcertObjectId() = default;
+
+	explicit FConcertObjectId(const UObject* InObject)
+		: FConcertObjectId(
+			*InObject->GetClass()->GetPathName(), 
+			InObject->GetPackage()->GetFName(), 
+			InObject->GetFName(), 
+			InObject->GetOuter() ? FName(*InObject->GetOuter()->GetPathName()) : FName(), 
+			InObject->GetExternalPackage() ? InObject->GetExternalPackage()->GetFName() : FName(), 
+			InObject->GetFlags())
 	{
 	}
 
-	explicit FConcertObjectId(const UObject* InObject)
-		: ObjectClassPathName(*InObject->GetClass()->GetPathName())
-		, ObjectPackageName(InObject->GetPackage()->GetFName())
-		, ObjectName(InObject->GetFName())
-		, ObjectOuterPathName(InObject->GetOuter() ? FName(*InObject->GetOuter()->GetPathName()) : FName())
-		, ObjectExternalPackageName(InObject->GetExternalPackage() ? InObject->GetExternalPackage()->GetFName() : FName())
-		, ObjectPersistentFlags(InObject->GetFlags() & RF_Load)
+	FConcertObjectId(const FTransactionObjectId& InObjectId, const uint32 InObjectFlags)
+		: FConcertObjectId(InObjectId.ObjectClassPathName, InObjectId.ObjectPackageName, InObjectId.ObjectName, InObjectId.ObjectOuterPathName, InObjectId.ObjectExternalPackageName, InObjectFlags)
 	{
 	}
 
-	FConcertObjectId(const FName InObjectClassPathName, const FName InObjectPackageName, const FName InObjectName, const FName InObjectOuterPathName, const FName InObjectExternalPackageName,  const uint32 InObjectFlags)
+	FConcertObjectId(const FName InObjectClassPathName, const FName InObjectPackageName, const FName InObjectName, const FName InObjectOuterPathName, const FName InObjectExternalPackageName, const uint32 InObjectFlags)
 		: ObjectClassPathName(InObjectClassPathName)
 		, ObjectPackageName(InObjectPackageName)
 		, ObjectName(InObjectName)
@@ -68,7 +72,7 @@ struct FConcertObjectId
 	FName ObjectExternalPackageName;
 
 	UPROPERTY()
-	uint32 ObjectPersistentFlags;
+	uint32 ObjectPersistentFlags = 0;
 };
 
 USTRUCT()
