Update the selection set when actors are deleted from a transaction

diff --git a/Engine/Plugins/Developer/Concert/ConcertSync/ConcertSyncClient/Source/ConcertSyncClient/ConcertSyncClient.Build.cs b/Engine/Plugins/Developer/Concert/ConcertSync/ConcertSyncClient/Source/ConcertSyncClient/ConcertSyncClient.Build.cs
--- a/Engine/Plugins/Developer/Concert/ConcertSync/ConcertSyncClient/Source/ConcertSyncClient/ConcertSyncClient.Build.cs
+++ b/Engine/Plugins/Developer/Concert/ConcertSync/ConcertSyncClient/Source/ConcertSyncClient/ConcertSyncClient.Build.cs
@@ -63,6 +63,8 @@ namespace UnrealBuildTool.Rules
 						"LevelEditor",
 						"VREditor",
 						"EditorStyle",
+						"TypedElementRuntime",
+						"TypedElementFramework"
 					}
 				);
 			}
diff --git a/Engine/Plugins/Developer/Concert/ConcertSync/ConcertSyncClient/Source/ConcertSyncClient/Private/ConcertClientTransactionBridge.cpp b/Engine/Plugins/Developer/Concert/ConcertSync/ConcertSyncClient/Source/ConcertSyncClient/Private/ConcertClientTransactionBridge.cpp
--- a/Engine/Plugins/Developer/Concert/ConcertSync/ConcertSyncClient/Source/ConcertSyncClient/Private/ConcertClientTransactionBridge.cpp
+++ b/Engine/Plugins/Developer/Concert/ConcertSync/ConcertSyncClient/Source/ConcertSyncClient/Private/ConcertClientTransactionBridge.cpp
@@ -19,6 +19,9 @@
 	#include "Editor/UnrealEdEngine.h"
 	#include "Editor/TransBuffer.h"
 	#include "Engine/Selection.h"
+
+	#include "Elements/Framework/EngineElementsLibrary.h"
+	#include "Elements/Framework/TypedElementSelectionSet.h"
 #endif
 
 
@@ -52,6 +55,59 @@ bool RunTransactionFilters(const TArray<FTransactionClassFilter>& InFilters, UOb
 	return bMatchFilter;
 }
 
+template <typename TUObjectType>
+FTypedElementHandle AcquireTypedElementHandle(const TUObjectType* InObject)
+{
+	if constexpr(std::is_same<AActor,TUObjectType>::value)
+	{
+		return UEngineElementsLibrary::AcquireEditorActorElementHandle(InObject, /*bAllowCreate*/false);
+	}
+	else if constexpr(std::is_same<UActorComponent,TUObjectType>::value)
+	{
+		return UEngineElementsLibrary::AcquireEditorComponentElementHandle(InObject, /*bAllowCreate*/false);
+	}
+	return {};
+}
+
+template <typename TUObjectArrayType>
+void DeselectElements(UTypedElementSelectionSet* InSelectionSet, const TUObjectArrayType& InUObjectArray, const FTypedElementSelectionOptions& SelectionOptions)
+{
+	using ElementType = typename TUObjectArrayType::ElementType;
+	static_assert(std::is_same<AActor*,ElementType>::value || std::is_same<UActorComponent*,ElementType>::value);
+	if (InUObjectArray.Num() == 0)
+	{
+		return;
+	}
+	TArray<FTypedElementHandle> HandlesToRemoveFromSelection;
+	for (ElementType Object: InUObjectArray)
+	{
+		if (FTypedElementHandle Handle = AcquireTypedElementHandle(Object))
+		{
+			HandlesToRemoveFromSelection.Add(Handle);
+		}
+	}
+	InSelectionSet->DeselectElements(HandlesToRemoveFromSelection, SelectionOptions);
+}
+
+void DeselectActorsAndActorComponents(const TArray<AActor*>& DeletedActors, const TArray<UActorComponent*>& DeletedActorComponents)
+{
+#if WITH_EDITOR
+	if (!GEditor)
+	{
+		return;
+	}
+
+	const FTypedElementSelectionOptions SelectionOptions = FTypedElementSelectionOptions()
+		.SetAllowHidden(true)
+		.SetAllowGroups(false)
+		.SetWarnIfLocked(false)
+		.SetChildElementInclusionMethod(ETypedElementChildInclusionMethod::Recursive);
+
+	DeselectElements(GEditor->GetSelectedActors()->GetElementSelectionSet(), DeletedActors, SelectionOptions);
+	DeselectElements(GEditor->GetSelectedComponents()->GetElementSelectionSet(), DeletedActorComponents, SelectionOptions);
+#endif
+}
+
 ETransactionFilterResult ApplyCustomFilter(const TMap<FName, FTransactionFilterDelegate>& CustomFilters, UObject* InObject, UPackage* InChangedPackage)
 {
 	for (const auto& Item : CustomFilters)
@@ -397,6 +453,8 @@ void ProcessTransactionEvent(const FConcertTransactionEventBase& InEvent, const
 	// --------------------------------------------------------------------------------------------------------------------
 	// Phase 4
 	// --------------------------------------------------------------------------------------------------------------------
+	TArray<AActor*> DeletedActorsForSelectionUpdate;
+	TArray<UActorComponent*> DeletedActorComponentForSelectionUpdate;
 	for (int32 ObjectIndex = SortedExportedObjects.Num() - 1; ObjectIndex >= 0; --ObjectIndex)
 	{
 		const ConcertSyncClientUtil::FGetObjectResult& TransactionObjectRef = TransactionObjects[ObjectIndex];
@@ -408,14 +466,19 @@ void ProcessTransactionEvent(const FConcertTransactionEventBase& InEvent, const
 			continue;
 		}
 
-		if (AActor* TransactionActor = Cast<AActor>(TransactionObject))
+		if (ObjectUpdate.ObjectData.bIsPendingKill)
 		{
-			if (ObjectUpdate.ObjectData.bIsPendingKill)
+			if (AActor* TransactionActor = Cast<AActor>(TransactionObject))
 			{
 				if (DeletedActors)
 				{
 					DeletedActors->Add(TransactionActor);
 				}
+				DeletedActorsForSelectionUpdate.Add(TransactionActor);
+			}
+			else if(UActorComponent* TransactionActorComponent = Cast<UActorComponent>(TransactionObject))
+			{
+				DeletedActorComponentForSelectionUpdate.Add(TransactionActorComponent);
 			}
 		}
 
@@ -465,6 +528,8 @@ void ProcessTransactionEvent(const FConcertTransactionEventBase& InEvent, const
 	}
 #endif
 
+	DeselectActorsAndActorComponents(DeletedActorsForSelectionUpdate, DeletedActorComponentForSelectionUpdate);
+
 	// TODO: This can sometimes cause deadlocks - need to investigate why
 	if (bObjectsDeleted)
 	{
