Add nullptr check.

#rnx

diff --git a/Engine/Plugins/Developer/Concert/ConcertSync/ConcertSyncClient/Source/ConcertSyncClient/Private/ConcertClientTransactionBridge.cpp b/Engine/Plugins/Developer/Concert/ConcertSync/ConcertSyncClient/Source/ConcertSyncClient/Private/ConcertClientTransactionBridge.cpp
--- a/Engine/Plugins/Developer/Concert/ConcertSync/ConcertSyncClient/Source/ConcertSyncClient/Private/ConcertClientTransactionBridge.cpp
+++ b/Engine/Plugins/Developer/Concert/ConcertSync/ConcertSyncClient/Source/ConcertSyncClient/Private/ConcertClientTransactionBridge.cpp
@@ -328,7 +328,7 @@ void ProcessTransactionEvent(const FConcertTransactionEventBase& InEvent, const
 				});
 			}
 
-			if (ObjectUpdate.ObjectData.bResetExisting && !NewlyCreatedObjects.Contains(TransactionObjectRef.Obj))
+			if (TransactionObjectRef.Obj && ObjectUpdate.ObjectData.bResetExisting && !NewlyCreatedObjects.Contains(TransactionObjectRef.Obj))
 			{
 				ConcertSyncUtil::ResetObjectPropertiesToArchetypeValues(TransactionObjectRef.Obj, bIncludeEditorOnlyProperties);
 			}
