Fixed Concert incorrectly skipping changes that only had transaction annotation data

#rnx

diff --git a/Engine/Plugins/Developer/Concert/ConcertSync/ConcertSyncClient/Source/ConcertSyncClient/Private/ConcertClientTransactionBridge.cpp b/Engine/Plugins/Developer/Concert/ConcertSync/ConcertSyncClient/Source/ConcertSyncClient/Private/ConcertClientTransactionBridge.cpp
--- a/Engine/Plugins/Developer/Concert/ConcertSync/ConcertSyncClient/Source/ConcertSyncClient/Private/ConcertClientTransactionBridge.cpp
+++ b/Engine/Plugins/Developer/Concert/ConcertSync/ConcertSyncClient/Source/ConcertSyncClient/Private/ConcertClientTransactionBridge.cpp
@@ -779,7 +779,7 @@ void FConcertClientTransactionBridge::HandleObjectTransacted(UObject* InObject,
 	TSharedPtr<ITransactionObjectAnnotation> TransactionAnnotation = InTransactionEvent.GetAnnotation();
 	const bool bUseSerializedAnnotationData = !bIncludeAnnotationObjectChanges || !TransactionAnnotation || !TransactionAnnotation->SupportsAdditionalObjectChanges();
 
-	if (!InTransactionEvent.HasNonPropertyChanges() && ExportedProperties.Num() == 0)
+	if (!InTransactionEvent.HasNonPropertyChanges() && ExportedProperties.Num() == 0 && !(TransactionAnnotation && bUseSerializedAnnotationData))
 	{
 		// This object has no changes to send
 		return;
