Fixed a partially destructed actor losing its components when restored via Concert

It's possible for components to have been GC'd without the actor being GC'd, if the actor is being referenced by the undo buffer. To ensure the actor is restored to a complete state we need to allow the creation of native components (as we can't re-run the C++ constructor), as well as re-running the construction script (to ensure that any SCS/UCS components are alive).

#rnx

diff --git a/Engine/Plugins/Developer/Concert/ConcertSync/ConcertSyncClient/Source/ConcertSyncClient/Private/ConcertClientTransactionBridge.cpp b/Engine/Plugins/Developer/Concert/ConcertSync/ConcertSyncClient/Source/ConcertSyncClient/Private/ConcertClientTransactionBridge.cpp
--- a/Engine/Plugins/Developer/Concert/ConcertSync/ConcertSyncClient/Source/ConcertSyncClient/Private/ConcertClientTransactionBridge.cpp
+++ b/Engine/Plugins/Developer/Concert/ConcertSync/ConcertSyncClient/Source/ConcertSyncClient/Private/ConcertClientTransactionBridge.cpp
@@ -318,19 +318,46 @@ void ProcessTransactionEvent(const FConcertTransactionEventBase& InEvent, const
 			TransactionObjectRef = ConcertSyncClientUtil::GetObject(ObjectUpdate.ObjectId, ObjectUpdate.ObjectData.NewName, ObjectUpdate.ObjectData.NewOuterPathName, ObjectUpdate.ObjectData.NewExternalPackageName, ObjectUpdate.ObjectData.bAllowCreate);
 			bObjectsDeleted |= (ObjectUpdate.ObjectData.bIsPendingKill || TransactionObjectRef.NeedsGC());
 
-			if (TransactionObjectRef.Obj && TransactionObjectRef.NewlyCreated())
+			if (TransactionObjectRef.Obj)
 			{
-				// Track this object (and any inner objects, as they must also be new) as newly created
-				NewlyCreatedObjects.Add(TransactionObjectRef.Obj);
-				ForEachObjectWithOuter(TransactionObjectRef.Obj, [&NewlyCreatedObjects](UObject* InnerObj)
+				if (TransactionObjectRef.NewlyCreated())
 				{
-					NewlyCreatedObjects.Add(InnerObj);
-				});
-			}
+					// Track this object (and any inner objects, as they must also be new) as newly created
+					NewlyCreatedObjects.Add(TransactionObjectRef.Obj);
+					ForEachObjectWithOuter(TransactionObjectRef.Obj, [&NewlyCreatedObjects](UObject* InnerObj)
+					{
+						NewlyCreatedObjects.Add(InnerObj);
+					});
+				}
+				else if (!ObjectUpdate.ObjectData.bIsPendingKill && !IsValid(TransactionObjectRef.Obj))
+				{
+					// If we're bringing this actor back to life, then make sure any SCS/UCS components exist in a clean state
+					// We have to do this as some of the actor components may have been GC'd when using "AllowEliminatingReferences(false)" (eg, within the transaction buffer)
+					if (AActor* Actor = Cast<AActor>(TransactionObjectRef.Obj))
+					{
+						TSet<const UObject*> ExistingSubObjects;
+						ForEachObjectWithOuter(Actor, [&ExistingSubObjects](UObject* InnerObj)
+						{
+							ExistingSubObjects.Add(InnerObj);
+						});
+
+						Actor->RerunConstructionScripts();
+
+						// Track any reconstructed objects as newly created
+						ForEachObjectWithOuter(Actor, [&NewlyCreatedObjects, &ExistingSubObjects](UObject* InnerObj)
+						{
+							if (!ExistingSubObjects.Contains(InnerObj))
+							{
+								NewlyCreatedObjects.Add(InnerObj);
+							}
+						});
+					}
+				}
 
-			if (TransactionObjectRef.Obj && ObjectUpdate.ObjectData.bResetExisting && !NewlyCreatedObjects.Contains(TransactionObjectRef.Obj))
-			{
-				ConcertSyncUtil::ResetObjectPropertiesToArchetypeValues(TransactionObjectRef.Obj, bIncludeEditorOnlyProperties);
+				if (ObjectUpdate.ObjectData.bResetExisting && !NewlyCreatedObjects.Contains(TransactionObjectRef.Obj))
+				{
+					ConcertSyncUtil::ResetObjectPropertiesToArchetypeValues(TransactionObjectRef.Obj, bIncludeEditorOnlyProperties);
+				}
 			}
 		}
 	}
@@ -705,6 +732,11 @@ void FConcertClientTransactionBridge::HandleObjectTransacted(UObject* InObject,
 		return;
 	}
 
+	if (!bIncludeEditorOnlyProperties && InObject->IsEditorOnly())
+	{
+		return;
+	}
+
 	UPackage* ChangedPackage = InObject->GetOutermost();
 	ETransactionFilterResult FilterResult = ConcertClientTransactionBridgeUtil::ApplyTransactionFilters(TransactionFilters, InObject, ChangedPackage);
 	FOngoingTransaction* TrackedTransaction = OngoingTransactions.Find(InTransactionEvent.GetOperationId());
@@ -758,7 +790,8 @@ void FConcertClientTransactionBridge::HandleObjectTransacted(UObject* InObject,
 		return;
 	}
 
-	bool bCanCreateOrRenameObject = true;
+	bool bCanCreateObject = true;
+	bool bCanRenameObject = true;
 	if (const UActorComponent* Component = Cast<UActorComponent>(InObject))
 	{
 		bool bIsOwnedByChildActor = false;
@@ -767,14 +800,17 @@ void FConcertClientTransactionBridge::HandleObjectTransacted(UObject* InObject,
 			bIsOwnedByChildActor = ComponentOwner->IsChildActor();
 		}
 
-		// Components that are managed by a native or Blueprint class cannot be created or renamed, so we only allow "instance" components to be created or renamed
-		bCanCreateOrRenameObject = !bIsOwnedByChildActor && Component->CreationMethod == EComponentCreationMethod::Instance;
+		// Components that are managed by a construction script cannot be created
+		bCanCreateObject = !bIsOwnedByChildActor && !Component->IsCreatedByConstructionScript();
+
+		// Components that are managed by a native or Blueprint class cannot be renamed, so we only allow "instance" components to be created or renamed
+		bCanRenameObject = !bIsOwnedByChildActor && Component->CreationMethod == EComponentCreationMethod::Instance;
 	}
 
-	const FName NewObjectPackageName = (bCanCreateOrRenameObject && (InTransactionEvent.HasOuterChange() || InTransactionEvent.HasExternalPackageChange())) ? InObject->GetPackage()->GetFName() : FName();
-	const FName NewObjectName = (bCanCreateOrRenameObject && InTransactionEvent.HasNameChange()) ? InObject->GetFName() : FName();
-	const FName NewObjectOuterPathName = (bCanCreateOrRenameObject && (InTransactionEvent.HasOuterChange() && InObject->GetOuter())) ? FName(*InObject->GetOuter()->GetPathName()) : FName();
-	const FName NewObjectExternalPackageName = (bCanCreateOrRenameObject && (InTransactionEvent.HasExternalPackageChange() && InObject->GetExternalPackage())) ? InObject->GetExternalPackage()->GetFName() : FName();
+	const FName NewObjectPackageName = (bCanRenameObject && (InTransactionEvent.HasOuterChange() || InTransactionEvent.HasExternalPackageChange())) ? InObject->GetPackage()->GetFName() : FName();
+	const FName NewObjectName = (bCanRenameObject && InTransactionEvent.HasNameChange()) ? InObject->GetFName() : FName();
+	const FName NewObjectOuterPathName = (bCanRenameObject && (InTransactionEvent.HasOuterChange() && InObject->GetOuter())) ? FName(*InObject->GetOuter()->GetPathName()) : FName();
+	const FName NewObjectExternalPackageName = (bCanRenameObject && (InTransactionEvent.HasExternalPackageChange() && InObject->GetExternalPackage())) ? InObject->GetExternalPackage()->GetFName() : FName();
 	const TArray<const FProperty*> ExportedProperties = ConcertSyncClientUtil::GetExportedProperties(InObject->GetClass(), InTransactionEvent.GetChangedProperties(), bIncludeEditorOnlyProperties);
 	TSharedPtr<ITransactionObjectAnnotation> TransactionAnnotation = InTransactionEvent.GetAnnotation();
 	const bool bUseSerializedAnnotationData = !bIncludeAnnotationObjectChanges || !TransactionAnnotation || !TransactionAnnotation->SupportsAdditionalObjectChanges();
@@ -853,7 +889,7 @@ void FConcertClientTransactionBridge::HandleObjectTransacted(UObject* InObject,
 		FConcertExportedObject& ObjectUpdate = OngoingTransaction.FinalizedData.FinalizedObjectUpdates.AddDefaulted_GetRef();
 		ObjectUpdate.ObjectId = ObjectId;
 		ObjectUpdate.ObjectPathDepth = ConcertSyncClientUtil::GetObjectPathDepth(InObject);
-		ObjectUpdate.ObjectData.bAllowCreate = bIsNewlyCreated && bCanCreateOrRenameObject;
+		ObjectUpdate.ObjectData.bAllowCreate = bIsNewlyCreated && bCanCreateObject;
 		ObjectUpdate.ObjectData.bResetExisting = bIsNewlyCreated;
 		ObjectUpdate.ObjectData.bIsPendingKill = !IsValid(InObject);
 		ObjectUpdate.ObjectData.NewPackageName = NewObjectPackageName;
