Ensure the actor isn't garbage prior to re-running its construction script when recreating it via Concert

#rnx

diff --git a/Engine/Plugins/Developer/Concert/ConcertSync/ConcertSyncClient/Source/ConcertSyncClient/Private/ConcertClientTransactionBridge.cpp b/Engine/Plugins/Developer/Concert/ConcertSync/ConcertSyncClient/Source/ConcertSyncClient/Private/ConcertClientTransactionBridge.cpp
--- a/Engine/Plugins/Developer/Concert/ConcertSync/ConcertSyncClient/Source/ConcertSyncClient/Private/ConcertClientTransactionBridge.cpp
+++ b/Engine/Plugins/Developer/Concert/ConcertSync/ConcertSyncClient/Source/ConcertSyncClient/Private/ConcertClientTransactionBridge.cpp
@@ -329,7 +329,12 @@ void ProcessTransactionEvent(const FConcertTransactionEventBase& InEvent, const
 						NewlyCreatedObjects.Add(InnerObj);
 					});
 				}
-				else if (!ObjectUpdate.ObjectData.bIsPendingKill && !IsValid(TransactionObjectRef.Obj))
+				
+				// Update the pending kill state
+				const bool bWasPendingKill = !IsValid(TransactionObjectRef.Obj);
+				ConcertSyncClientUtil::UpdatePendingKillState(TransactionObjectRef.Obj, ObjectUpdate.ObjectData.bIsPendingKill);
+
+				if (!TransactionObjectRef.NewlyCreated() && !ObjectUpdate.ObjectData.bIsPendingKill && bWasPendingKill)
 				{
 					// If we're bringing this actor back to life, then make sure any SCS/UCS components exist in a clean state
 					// We have to do this as some of the actor components may have been GC'd when using "AllowEliminatingReferences(false)" (eg, within the transaction buffer)
@@ -449,9 +454,6 @@ void ProcessTransactionEvent(const FConcertTransactionEventBase& InEvent, const
 			continue;
 		}
 
-		// Update the pending kill state
-		ConcertSyncClientUtil::UpdatePendingKillState(TransactionObject, ObjectUpdate.ObjectData.bIsPendingKill);
-
 		// Apply the new data
 		if (ObjectUpdate.ObjectData.SerializedData.Num() > 0)
 		{
