Ensure non-realtime viewports update when applying a change from Concert

#rnx

diff --git a/Engine/Plugins/Developer/Concert/ConcertSync/ConcertSyncClient/Source/ConcertSyncClient/Private/ConcertClientTransactionBridge.cpp b/Engine/Plugins/Developer/Concert/ConcertSync/ConcertSyncClient/Source/ConcertSyncClient/Private/ConcertClientTransactionBridge.cpp
--- a/Engine/Plugins/Developer/Concert/ConcertSync/ConcertSyncClient/Source/ConcertSyncClient/Private/ConcertClientTransactionBridge.cpp
+++ b/Engine/Plugins/Developer/Concert/ConcertSync/ConcertSyncClient/Source/ConcertSyncClient/Private/ConcertClientTransactionBridge.cpp
@@ -561,9 +561,13 @@ void ProcessTransactionEvent(const FConcertTransactionEventBase& InEvent, const
 	}
 
 #if WITH_EDITOR
-	if (bIsSnapshot && GUnrealEd)
+	if (GUnrealEd)
 	{
-		GUnrealEd->UpdatePivotLocationForSelection();
+		if (bIsSnapshot)
+		{
+			GUnrealEd->UpdatePivotLocationForSelection();
+		}
+		GUnrealEd->RedrawLevelEditingViewports();
 	}
 #endif
 }
