Improved logic for destroying/restoring actors in Concert, without needing to sync their owner level

We now call DestroyActor when marking an actor as garbage in Concert (which removes it from its owner level), and make sure that a restored actor is re-added to its level. This avoids needing to sync the (non-property) Actors and ActorsForGC arrays of the owner level.

#rnx

diff --git a/Engine/Plugins/Developer/Concert/ConcertSync/ConcertSyncClient/Source/ConcertSyncClient/Private/ConcertClientTransactionBridge.cpp b/Engine/Plugins/Developer/Concert/ConcertSync/ConcertSyncClient/Source/ConcertSyncClient/Private/ConcertClientTransactionBridge.cpp
--- a/Engine/Plugins/Developer/Concert/ConcertSync/ConcertSyncClient/Source/ConcertSyncClient/Private/ConcertClientTransactionBridge.cpp
+++ b/Engine/Plugins/Developer/Concert/ConcertSync/ConcertSyncClient/Source/ConcertSyncClient/Private/ConcertClientTransactionBridge.cpp
@@ -292,6 +292,7 @@ void ProcessTransactionEvent(const FConcertTransactionEventBase& InEvent, const
 	// Phase 1
 	// --------------------------------------------------------------------------------------------------------------------
 	bool bObjectsDeleted = false;
+	TArray<AActor*> ResurrectedActors;
 	TArray<ConcertSyncClientUtil::FGetObjectResult, TInlineAllocator<32>> TransactionObjects;
 	{
 		TSet<const UObject*> NewlyCreatedObjects;
@@ -345,6 +346,8 @@ void ProcessTransactionEvent(const FConcertTransactionEventBase& InEvent, const
 					// We have to do this as some of the actor components may have been GC'd when using "AllowEliminatingReferences(false)" (eg, within the transaction buffer)
 					if (AActor* Actor = Cast<AActor>(TransactionObjectRef.Obj))
 					{
+						ResurrectedActors.Add(Actor);
+
 						TSet<const UObject*> ExistingSubObjects;
 						ForEachObjectWithOuter(Actor, [&ExistingSubObjects](UObject* InnerObj)
 						{
@@ -552,6 +555,12 @@ void ProcessTransactionEvent(const FConcertTransactionEventBase& InEvent, const
 	}
 #endif
 
+	// Ensure that any actors we restored from the dead are added back to the actors array of their owner level
+	for (AActor* Actor : ResurrectedActors)
+	{
+		ConcertSyncClientUtil::AddActorToOwnerLevel(Actor);
+	}
+
 	DeselectActorsAndActorComponents(DeletedActorsForSelectionUpdate, DeletedActorComponentForSelectionUpdate);
 
 	// TODO: This can sometimes cause deadlocks - need to investigate why
diff --git a/Engine/Plugins/Developer/Concert/ConcertSync/ConcertSyncClient/Source/ConcertSyncClient/Private/ConcertSyncClientUtil.cpp b/Engine/Plugins/Developer/Concert/ConcertSync/ConcertSyncClient/Source/ConcertSyncClient/Private/ConcertSyncClientUtil.cpp
--- a/Engine/Plugins/Developer/Concert/ConcertSync/ConcertSyncClient/Source/ConcertSyncClient/Private/ConcertSyncClientUtil.cpp
+++ b/Engine/Plugins/Developer/Concert/ConcertSync/ConcertSyncClient/Source/ConcertSyncClient/Private/ConcertSyncClientUtil.cpp
@@ -114,9 +114,37 @@ bool CanPerformBlockingAction(const bool bBlockDuringInteraction)
 
 void UpdatePendingKillState(UObject* InObj, const bool bIsPendingKill)
 {
+	const bool bWasPendingKill = !IsValid(InObj);
+	if (bIsPendingKill == bWasPendingKill)
+	{
+		return;
+	}
+
 	if (bIsPendingKill)
 	{
-		InObj->MarkAsGarbage();
+		bool bMarkAsGarbage = true;
+
+		if (AActor* Actor = Cast<AActor>(InObj))
+		{
+			if (UWorld* ActorWorld = Actor->GetWorld())
+			{
+#if WITH_EDITOR
+				if (GIsEditor)
+				{
+					bMarkAsGarbage = !ActorWorld->EditorDestroyActor(Actor, /*bShouldModifyLevel*/false);
+				}
+				else
+#endif	// WITH_EDITOR
+				{
+					bMarkAsGarbage = !ActorWorld->DestroyActor(Actor, /*bNetForce*/false, /*bShouldModifyLevel*/false);
+				}
+			}
+		}
+
+		if (bMarkAsGarbage)
+		{
+			InObj->MarkAsGarbage();
+		}
 	}
 	else
 	{
@@ -124,6 +152,37 @@ void UpdatePendingKillState(UObject* InObj, const bool bIsPendingKill)
 	}
 }
 
+void AddActorToOwnerLevel(AActor* InActor)
+{
+	if (ULevel* Level = InActor->GetLevel())
+	{
+		if (Level->Actors.Contains(InActor))
+		{
+			return;
+		}
+
+		Level->Actors.Add(InActor);
+		Level->ActorsForGC.Add(InActor);
+
+#if WITH_EDITOR
+		if (GIsEditor)
+		{
+			if (Level->IsUsingActorFolders())
+			{
+				InActor->FixupActorFolder();
+			}
+
+			GEngine->BroadcastLevelActorAdded(InActor);
+
+			if (UWorld* World = Level->GetWorld())
+			{
+				World->BroadcastLevelsChanged();
+			}
+		}
+#endif	// WITH_EDITOR
+	}
+}
+
 bool ObjectIdsMatch(const FConcertObjectId& One, const FConcertObjectId& Two)
 {
 	return One.ObjectClassPathName == Two.ObjectClassPathName
@@ -217,7 +276,7 @@ FGetObjectResult GetObject(const FConcertObjectId& InObjectId, const FName InNew
 					if (UObject* NewObject = StaticFindObject(ObjectClass, NewObjectOuter ? NewObjectOuter : ExistingObjectOuter, *ObjectNameToCreate.ToString(), /*bExactClass*/true))
 					{
 						UE_LOG(LogConcert, Warning, TEXT("Attempted to rename '%s' over '%s'. Re-using the found object instead of performing the rename!"), *ExistingObject->GetPathName(), *NewObject->GetPathName());
-						ExistingObject->MarkAsGarbage();
+						UpdatePendingKillState(ExistingObject, /*bIsPendingKill*/true);
 						ResultFlags |= EGetObjectResultFlags::NeedsGC;
 
 						ExistingObject = NewObject;
diff --git a/Engine/Plugins/Developer/Concert/ConcertSync/ConcertSyncClient/Source/ConcertSyncClient/Private/ConcertSyncClientUtil.h b/Engine/Plugins/Developer/Concert/ConcertSync/ConcertSyncClient/Source/ConcertSyncClient/Private/ConcertSyncClientUtil.h
--- a/Engine/Plugins/Developer/Concert/ConcertSync/ConcertSyncClient/Source/ConcertSyncClient/Private/ConcertSyncClientUtil.h
+++ b/Engine/Plugins/Developer/Concert/ConcertSync/ConcertSyncClient/Source/ConcertSyncClient/Private/ConcertSyncClientUtil.h
@@ -75,6 +75,8 @@ namespace ConcertSyncClientUtil
 
 	void UpdatePendingKillState(UObject* InObj, const bool bIsPendingKill);
 
+	void AddActorToOwnerLevel(AActor* InActor);
+
 	bool ObjectIdsMatch(const FConcertObjectId& One, const FConcertObjectId& Two);
 
 	int32 GetObjectPathDepth(UObject* InObjToTest);
