Added an option to allow Concert to skip sending non-property object data

#rnx

diff --git a/Engine/Plugins/Developer/Concert/ConcertSync/ConcertSyncClient/Source/ConcertSyncClient/Private/ConcertClientTransactionBridge.cpp b/Engine/Plugins/Developer/Concert/ConcertSync/ConcertSyncClient/Source/ConcertSyncClient/Private/ConcertClientTransactionBridge.cpp
--- a/Engine/Plugins/Developer/Concert/ConcertSync/ConcertSyncClient/Source/ConcertSyncClient/Private/ConcertClientTransactionBridge.cpp
+++ b/Engine/Plugins/Developer/Concert/ConcertSync/ConcertSyncClient/Source/ConcertSyncClient/Private/ConcertClientTransactionBridge.cpp
@@ -592,6 +592,7 @@ FConcertClientTransactionBridge::FConcertClientTransactionBridge()
 	: bHasBoundUnderlyingLocalTransactionEvents(false)
 	, bIgnoreLocalTransactions(false)
 	, bIncludeEditorOnlyProperties(true)
+	, bIncludeNonPropertyObjectData(true)
 	, bIncludeAnnotationObjectChanges(GetDefault<UConcertSyncConfig>()->bIncludeAnnotationObjectChanges)
 {
 	ConditionalBindUnderlyingLocalTransactionEvents();
@@ -623,6 +624,11 @@ void FConcertClientTransactionBridge::SetIncludeEditorOnlyProperties(const bool
 	bIncludeEditorOnlyProperties = InIncludeEditorOnlyProperties;
 }
 
+void FConcertClientTransactionBridge::SetIncludeNonPropertyObjectData(const bool InIncludeNonPropertyObjectData)
+{
+	bIncludeNonPropertyObjectData = InIncludeNonPropertyObjectData;
+}
+
 void FConcertClientTransactionBridge::SetIncludeAnnotationObjectChanges(const bool InIncludeAnnotationObjectChanges)
 {
 	bIncludeAnnotationObjectChanges = InIncludeAnnotationObjectChanges;
@@ -831,7 +837,12 @@ void FConcertClientTransactionBridge::HandleObjectTransacted(UObject* InObject,
 	TSharedPtr<ITransactionObjectAnnotation> TransactionAnnotation = InTransactionEvent.GetAnnotation();
 	const bool bUseSerializedAnnotationData = !bIncludeAnnotationObjectChanges || !TransactionAnnotation || !TransactionAnnotation->SupportsAdditionalObjectChanges();
 
-	if (!InTransactionEvent.HasNonPropertyChanges() && ExportedProperties.Num() == 0 && !(TransactionAnnotation && bUseSerializedAnnotationData))
+	const bool bObjectHasChangesToSend = InTransactionEvent.HasIdOrPendingKillChanges()
+		|| (bIncludeNonPropertyObjectData && InTransactionEvent.HasNonPropertyChanges(/*SerializationOnly*/true))
+		|| ExportedProperties.Num() > 0
+		|| (TransactionAnnotation && bUseSerializedAnnotationData);
+
+	if (!bObjectHasChangesToSend)
 	{
 		// This object has no changes to send
 		return;
@@ -919,7 +930,7 @@ void FConcertClientTransactionBridge::HandleObjectTransacted(UObject* InObject,
 			TransactionAnnotation->Serialize(AnnotationWriter);
 		}
 
-		if (InTransactionEvent.HasNonPropertyChanges(/*SerializationOnly*/true))
+		if (bIncludeNonPropertyObjectData && InTransactionEvent.HasNonPropertyChanges(/*SerializationOnly*/true))
 		{
 			// The 'non-property changes' refers to custom data added by a deriver UObject before and/or after the standard serialized data. Since this is a custom
 			// data format, we don't know what changed, call the object to re-serialize this part, but still send the delta for the generic reflected properties (in RootPropertyNames).
diff --git a/Engine/Plugins/Developer/Concert/ConcertSync/ConcertSyncClient/Source/ConcertSyncClient/Private/ConcertClientTransactionBridge.h b/Engine/Plugins/Developer/Concert/ConcertSync/ConcertSyncClient/Source/ConcertSyncClient/Private/ConcertClientTransactionBridge.h
--- a/Engine/Plugins/Developer/Concert/ConcertSync/ConcertSyncClient/Source/ConcertSyncClient/Private/ConcertClientTransactionBridge.h
+++ b/Engine/Plugins/Developer/Concert/ConcertSync/ConcertSyncClient/Source/ConcertSyncClient/Private/ConcertClientTransactionBridge.h
@@ -14,6 +14,7 @@ public:
 
 	//~ IConcertClientTransactionBridge interface
 	virtual void SetIncludeEditorOnlyProperties(const bool InIncludeEditorOnlyProperties) override;
+	virtual void SetIncludeNonPropertyObjectData(const bool InIncludeNonPropertyObjectData) override;
 	virtual void SetIncludeAnnotationObjectChanges(const bool InIncludeAnnotationObjectChanges) override;
 	virtual FOnConcertClientLocalTransactionSnapshot& OnLocalTransactionSnapshot() override;
 	virtual FOnConcertClientLocalTransactionFinalized& OnLocalTransactionFinalized() override;
@@ -84,6 +85,9 @@ private:
 	/** Include non-cooked properties in object serialization */
 	bool bIncludeEditorOnlyProperties;
 
+	/** Include non-property object data in updates */
+	bool bIncludeNonPropertyObjectData;
+
 	/** Include object changes that have been generated via a transaction annotation (where possible) */
 	bool bIncludeAnnotationObjectChanges;
 
diff --git a/Engine/Plugins/Developer/Concert/ConcertSync/ConcertSyncClient/Source/ConcertSyncClient/Public/IConcertClientTransactionBridge.h b/Engine/Plugins/Developer/Concert/ConcertSync/ConcertSyncClient/Source/ConcertSyncClient/Public/IConcertClientTransactionBridge.h
--- a/Engine/Plugins/Developer/Concert/ConcertSync/ConcertSyncClient/Source/ConcertSyncClient/Public/IConcertClientTransactionBridge.h
+++ b/Engine/Plugins/Developer/Concert/ConcertSync/ConcertSyncClient/Source/ConcertSyncClient/Public/IConcertClientTransactionBridge.h
@@ -93,6 +93,12 @@ public:
 	 */
 	virtual void SetIncludeEditorOnlyProperties(const bool InIncludeEditorOnlyProperties) = 0;
 
+	/**
+	 * Set whether to include non-property object data in updates, or whether to only include property changes.
+	 * @note This is set to true by default.
+	 */
+	virtual void SetIncludeNonPropertyObjectData(const bool InIncludeNonPropertyObjectData) = 0;
+
 	/**
 	 * Set whether to include object changes that have been generated via a transaction annotation 
 	 * (where possible), or whether to send the entire transaction annotation blob instead.
diff --git a/Engine/Source/Runtime/CoreUObject/Public/Misc/TransactionObjectEvent.h b/Engine/Source/Runtime/CoreUObject/Public/Misc/TransactionObjectEvent.h
--- a/Engine/Source/Runtime/CoreUObject/Public/Misc/TransactionObjectEvent.h
+++ b/Engine/Source/Runtime/CoreUObject/Public/Misc/TransactionObjectEvent.h
@@ -287,10 +287,16 @@ public:
 		return ObjectChange.OriginalId.ObjectExternalPackageName;
 	}
 
-	/** Were any non-property changes made to the object? */
+	/** Were ID (name, outer, package) or pending kill changes made to the object? */
+	bool HasIdOrPendingKillChanges() const
+	{
+		return ObjectChange.DeltaChange.bHasNameChange || ObjectChange.DeltaChange.bHasOuterChange || ObjectChange.DeltaChange.bHasExternalPackageChange || ObjectChange.DeltaChange.bHasPendingKillChange;
+	}
+
+	/** Were any non-property changes made to the object? (name, outer, package, pending kill, or serialized non-property data) */
 	bool HasNonPropertyChanges(const bool InSerializationOnly = false) const
 	{
-		return (!InSerializationOnly && (ObjectChange.DeltaChange.bHasNameChange || ObjectChange.DeltaChange.bHasOuterChange || ObjectChange.DeltaChange.bHasExternalPackageChange || ObjectChange.DeltaChange.bHasPendingKillChange)) || ObjectChange.DeltaChange.bHasNonPropertyChanges;
+		return (!InSerializationOnly && HasIdOrPendingKillChanges()) || ObjectChange.DeltaChange.bHasNonPropertyChanges;
 	}
 
 	/** Were any property changes made to the object? */
