World Partition: remove concert specific code to keep deleted actors alive until save and move the ownership of deleted actors in the world partition object.
It is now possible because we are now calling DestroyActor instead of simply marking actors as pending kill since 21825302.

#rnx

diff --git a/Engine/Plugins/Developer/Concert/ConcertSync/ConcertSyncClient/Source/ConcertSyncClient/Private/ConcertClientTransactionBridge.cpp b/Engine/Plugins/Developer/Concert/ConcertSync/ConcertSyncClient/Source/ConcertSyncClient/Private/ConcertClientTransactionBridge.cpp
--- a/Engine/Plugins/Developer/Concert/ConcertSync/ConcertSyncClient/Source/ConcertSyncClient/Private/ConcertClientTransactionBridge.cpp
+++ b/Engine/Plugins/Developer/Concert/ConcertSync/ConcertSyncClient/Source/ConcertSyncClient/Private/ConcertClientTransactionBridge.cpp
@@ -262,7 +262,7 @@ struct FEditorTransactionNotification
 };
 #endif
 
-void ProcessTransactionEvent(const FConcertTransactionEventBase& InEvent, const FConcertSessionVersionInfo* InVersionInfo, const TArray<FName>& InPackagesToProcess, const FConcertLocalIdentifierTable* InLocalIdentifierTablePtr, const bool bIsSnapshot, const FConcertSyncWorldRemapper& WorldRemapper, TSet<AActor*>* DeletedActors, const bool bIncludeEditorOnlyProperties)
+void ProcessTransactionEvent(const FConcertTransactionEventBase& InEvent, const FConcertSessionVersionInfo* InVersionInfo, const TArray<FName>& InPackagesToProcess, const FConcertLocalIdentifierTable* InLocalIdentifierTablePtr, const bool bIsSnapshot, const FConcertSyncWorldRemapper& WorldRemapper, const bool bIncludeEditorOnlyProperties)
 {
 	// Transactions are applied in multiple-phases...
 	//	0) Sort the objects to be processed; creation needs to happen parent->child, and update needs to happen child->parent
@@ -497,10 +497,6 @@ void ProcessTransactionEvent(const FConcertTransactionEventBase& InEvent, const
 		{
 			if (AActor* TransactionActor = Cast<AActor>(TransactionObject))
 			{
-				if (DeletedActors)
-				{
-					DeletedActors->Add(TransactionActor);
-				}
 				DeletedActorsForSelectionUpdate.Add(TransactionActor);
 			}
 			else if(UActorComponent* TransactionActorComponent = Cast<UActorComponent>(TransactionObject))
@@ -656,17 +652,12 @@ FOnApplyTransaction& FConcertClientTransactionBridge::OnApplyTransaction()
 
 void FConcertClientTransactionBridge::ApplyRemoteTransaction(const FConcertTransactionEventBase& InEvent, const FConcertSessionVersionInfo* InVersionInfo, const TArray<FName>& InPackagesToProcess, const FConcertLocalIdentifierTable* InLocalIdentifierTablePtr, const bool bIsSnapshot)
 {
-	ConcertClientTransactionBridgeUtil::ProcessTransactionEvent(InEvent, InVersionInfo, InPackagesToProcess, InLocalIdentifierTablePtr, bIsSnapshot, FConcertSyncWorldRemapper(), nullptr, bIncludeEditorOnlyProperties);
+	ConcertClientTransactionBridgeUtil::ProcessTransactionEvent(InEvent, InVersionInfo, InPackagesToProcess, InLocalIdentifierTablePtr, bIsSnapshot, FConcertSyncWorldRemapper(), bIncludeEditorOnlyProperties);
 }
 
 void FConcertClientTransactionBridge::ApplyRemoteTransaction(const FConcertTransactionEventBase& InEvent, const FConcertSessionVersionInfo* InVersionInfo, const TArray<FName>& InPackagesToProcess, const FConcertLocalIdentifierTable* InLocalIdentifierTablePtr, const bool bIsSnapshot, const FConcertSyncWorldRemapper& ConcertSyncWorldRemapper)
 {
-	ConcertClientTransactionBridgeUtil::ProcessTransactionEvent(InEvent, InVersionInfo, InPackagesToProcess, InLocalIdentifierTablePtr, bIsSnapshot, ConcertSyncWorldRemapper, nullptr, bIncludeEditorOnlyProperties);
-}
-
-void FConcertClientTransactionBridge::ApplyRemoteTransaction(const FConcertTransactionEventBase& InEvent, const FConcertSessionVersionInfo* InVersionInfo, const TArray<FName>& InPackagesToProcess, const FConcertLocalIdentifierTable* InLocalIdentifierTablePtr, const bool bIsSnapshot, const FConcertSyncWorldRemapper& ConcertSyncWorldRemapper, TSet<AActor*>* DeletedActors)
-{
-	ConcertClientTransactionBridgeUtil::ProcessTransactionEvent(InEvent, InVersionInfo, InPackagesToProcess, InLocalIdentifierTablePtr, bIsSnapshot, ConcertSyncWorldRemapper, DeletedActors, bIncludeEditorOnlyProperties);
+	ConcertClientTransactionBridgeUtil::ProcessTransactionEvent(InEvent, InVersionInfo, InPackagesToProcess, InLocalIdentifierTablePtr, bIsSnapshot, ConcertSyncWorldRemapper, bIncludeEditorOnlyProperties);
 }
 
 bool& FConcertClientTransactionBridge::GetIgnoreLocalTransactionsRef()
diff --git a/Engine/Plugins/Developer/Concert/ConcertSync/ConcertSyncClient/Source/ConcertSyncClient/Private/ConcertClientTransactionBridge.h b/Engine/Plugins/Developer/Concert/ConcertSync/ConcertSyncClient/Source/ConcertSyncClient/Private/ConcertClientTransactionBridge.h
--- a/Engine/Plugins/Developer/Concert/ConcertSync/ConcertSyncClient/Source/ConcertSyncClient/Private/ConcertClientTransactionBridge.h
+++ b/Engine/Plugins/Developer/Concert/ConcertSync/ConcertSyncClient/Source/ConcertSyncClient/Private/ConcertClientTransactionBridge.h
@@ -23,7 +23,6 @@ public:
 
 	virtual void ApplyRemoteTransaction(const FConcertTransactionEventBase& InEvent, const FConcertSessionVersionInfo* InVersionInfo, const TArray<FName>& InPackagesToProcess, const FConcertLocalIdentifierTable* InLocalIdentifierTablePtr, const bool bIsSnapshot) override;
 	virtual void ApplyRemoteTransaction(const FConcertTransactionEventBase& InEvent, const FConcertSessionVersionInfo* InVersionInfo, const TArray<FName>& InPackagesToProcess, const FConcertLocalIdentifierTable* InLocalIdentifierTablePtr, const bool bIsSnapshot, const FConcertSyncWorldRemapper& ConcertSyncWorldRemapper) override;
-	virtual void ApplyRemoteTransaction(const FConcertTransactionEventBase& InEvent, const FConcertSessionVersionInfo* InVersionInfo, const TArray<FName>& InPackagesToProcess, const FConcertLocalIdentifierTable* InLocalIdentifierTablePtr, const bool bIsSnapshot, const class FConcertSyncWorldRemapper& ConcertSyncWorldRemapper, TSet<AActor*>* DeletedActors) override;
 	virtual bool& GetIgnoreLocalTransactionsRef() override;
 
 	virtual void RegisterTransactionFilter(FName FilterName, FTransactionFilterDelegate FilterHandle) override;
diff --git a/Engine/Plugins/Developer/Concert/ConcertSync/ConcertSyncClient/Source/ConcertSyncClient/Public/IConcertClientTransactionBridge.h b/Engine/Plugins/Developer/Concert/ConcertSync/ConcertSyncClient/Source/ConcertSyncClient/Public/IConcertClientTransactionBridge.h
--- a/Engine/Plugins/Developer/Concert/ConcertSync/ConcertSyncClient/Source/ConcertSyncClient/Public/IConcertClientTransactionBridge.h
+++ b/Engine/Plugins/Developer/Concert/ConcertSync/ConcertSyncClient/Source/ConcertSyncClient/Public/IConcertClientTransactionBridge.h
@@ -151,18 +151,6 @@ public:
 	 */
 	virtual void ApplyRemoteTransaction(const FConcertTransactionEventBase& InEvent, const FConcertSessionVersionInfo* InVersionInfo, const TArray<FName>& InPackagesToProcess, const FConcertLocalIdentifierTable* InLocalIdentifierTablePtr, const bool bIsSnapshot, const class FConcertSyncWorldRemapper& ConcertSyncWorldRemapper) = 0;
 
-	/**
-	 * Apply a remote transaction event to this local instance.
-	 * @param InEvent					The event to apply.
-	 * @param InVersionInfo				The version information for the serialized data in the event, or null if the event should be serialized using the compiled in version info.
-	 * @param InPackagesToProcess		The list of packages to apply changes for, or an empty array to apply all changes.
-	 * @param InLocalIdentifierTablePtr The local identifier table for the event data (if any).
-	 * @param bIsSnapshot				True if this transaction event was a snapshot rather than a finalized transaction.
-	 * @param ConcertSyncWorldRemapper	Remapper to use in the case the current world is different from the one sent in the transation.
-	 * @param DeletedActors				Optional set of actor pointers to gather deleted actors.
-	 */
-	virtual void ApplyRemoteTransaction(const FConcertTransactionEventBase& InEvent, const FConcertSessionVersionInfo* InVersionInfo, const TArray<FName>& InPackagesToProcess, const FConcertLocalIdentifierTable* InLocalIdentifierTablePtr, const bool bIsSnapshot, const class FConcertSyncWorldRemapper& ConcertSyncWorldRemapper, TSet<AActor*>* DeletedActors) = 0;
-
 	/** Callback to register delegate for handling transaction events */
 	virtual void RegisterTransactionFilter(FName FilterName, FTransactionFilterDelegate FilterHandle) = 0;
 
diff --git a/Engine/Source/Runtime/Engine/Private/WorldPartition/WorldPartition.cpp b/Engine/Source/Runtime/Engine/Private/WorldPartition/WorldPartition.cpp
--- a/Engine/Source/Runtime/Engine/Private/WorldPartition/WorldPartition.cpp
+++ b/Engine/Source/Runtime/Engine/Private/WorldPartition/WorldPartition.cpp
@@ -274,16 +274,13 @@ void UWorldPartition::OnPackageDirtyStateChanged(UPackage* Package)
 {
 	auto ShouldHandleActor = [this](AActor* Actor) {return Actor && Actor->IsMainPackageActor() && (Actor->GetLevel() != nullptr) && IsActorDescHandled(Actor);};
 
-	if (IsStreamingEnabled())
+	if (AActor* Actor = AActor::FindActorInPackage(Package); ShouldHandleActor(Actor))
 	{
-		if (AActor* Actor = AActor::FindActorInPackage(Package); ShouldHandleActor(Actor))
+		if (FWorldPartitionHandle ActorHandle(this, Actor->GetActorGuid()); ActorHandle.IsValid())
 		{
-			if (FWorldPartitionHandle ActorHandle(this, Actor->GetActorGuid()); ActorHandle.IsValid())
+			if (Package->IsDirty())
 			{
-				if (Package->IsDirty())
-				{
-					DirtyActors.Add(ActorHandle.ToReference(), Actor);
-				}
+				DirtyActors.Add(ActorHandle.ToReference(), Actor);
 			}
 		}
 	}
@@ -1171,7 +1168,8 @@ void UWorldPartition::Tick(float DeltaSeconds)
 		{
 			if (!DirtyActorIt.Key().IsValid() || !DirtyActorIt.Value()->GetPackage()->IsDirty())
 			{
-				if (DirtyActorIt.Key().IsValid() && DirtyActorIt.Key()->GetHardRefCount() <= 1)
+				// If we hold the last reference to that actor (or no reference are held at all), pin it to avoid unloading
+				if (IsStreamingEnabled() && DirtyActorIt.Key().IsValid() && DirtyActorIt.Key()->GetHardRefCount() <= 1)
 				{
 					PinnedActors->AddActors({ DirtyActorIt.Key().ToHandle() });
 				}
