Do not use static initialization of TSharedRef<..> in Concert source control proxy.

diff --git a/Engine/Plugins/Developer/Concert/ConcertSync/ConcertSyncClient/Source/ConcertSyncClient/Private/ConcertSourceControlProxy.cpp b/Engine/Plugins/Developer/Concert/ConcertSync/ConcertSyncClient/Source/ConcertSyncClient/Private/ConcertSourceControlProxy.cpp
--- a/Engine/Plugins/Developer/Concert/ConcertSync/ConcertSyncClient/Source/ConcertSyncClient/Private/ConcertSourceControlProxy.cpp
+++ b/Engine/Plugins/Developer/Concert/ConcertSync/ConcertSyncClient/Source/ConcertSyncClient/Private/ConcertSourceControlProxy.cpp
@@ -247,8 +247,8 @@ const TArray<FSourceControlStateRef>& FConcertSourceControlChangelistStateProxy:
 
 FSourceControlChangelistRef FConcertSourceControlChangelistStateProxy::GetChangelist() const
 {
-	static const FSourceControlChangelistRef DummyResult;
-	return ActualState.IsValid() ? ActualState->GetChangelist() : DummyResult;
+	static ISourceControlChangelist DummyCL;
+	return ActualState.IsValid() ? ActualState->GetChangelist() : DummyCL.AsShared();
 }
 
 FConcertSourceControlProxy::FConcertSourceControlProxy()
