MultiUser: Allow sequence players to free loop without synchronization point via the sequence manager.

diff --git a/Engine/Plugins/Developer/Concert/ConcertSync/ConcertSyncClient/Source/ConcertSyncClient/Private/ConcertClientSequencerManager.cpp b/Engine/Plugins/Developer/Concert/ConcertSync/ConcertSyncClient/Source/ConcertSyncClient/Private/ConcertClientSequencerManager.cpp
--- a/Engine/Plugins/Developer/Concert/ConcertSync/ConcertSyncClient/Source/ConcertSyncClient/Private/ConcertClientSequencerManager.cpp
+++ b/Engine/Plugins/Developer/Concert/ConcertSync/ConcertSyncClient/Source/ConcertSyncClient/Private/ConcertClientSequencerManager.cpp
@@ -28,6 +28,7 @@
 #if WITH_EDITOR
 	#include "ISequencerModule.h"
 	#include "ISequencer.h"
+	#include "SequencerSettings.h"
 	#include "Subsystems/AssetEditorSubsystem.h"
 	#include "Editor.h"
 #endif
@@ -42,6 +43,8 @@
 // Enable Sequence Playing on game client
 static TAutoConsoleVariable<int32> CVarEnableSequencePlayer(TEXT("Concert.EnableSequencePlayer"), 1, TEXT("Enable Concert Sequence Players on `-game` client."));
 
+static TAutoConsoleVariable<int32> CVarEnableLoopingOnPlayer(TEXT("Concert.EnableLoopingOnPlayer"), 1, TEXT("Enable Looping Sequence Players when sequencer looping is enabled."));
+
 // Enable opening Sequencer on remote machine whenever a Sequencer is opened, if both instances have this option on.
 static TAutoConsoleVariable<int32> CVarEnableRemoteSequencerOpen(TEXT("Concert.EnableOpenRemoteSequencer"), 1, TEXT("Enable Concert remote Sequencer opening."));
 
@@ -256,6 +259,96 @@
 	return CVarAlwaysCloseGamePlayerOnCloseEvent.GetValueOnAnyThread() > 0;
 }
 
+namespace UE::Private::ConcertClientSequencerManager
+{
+
+void ApplyPlayRangeToPlayer(ULevelSequencePlayer* Player, const FFrameNumberRange& PlayRange)
+{
+	check(Player);
+
+	const bool bLowerBoundClosed = PlayRange.GetLowerBound().IsClosed();
+	const bool bUpperBoundClosed = PlayRange.GetUpperBound().IsClosed();
+
+	if (!ensureMsgf(bLowerBoundClosed, TEXT("PlayRange lower bound is open, which is not supported.")))
+	{
+		return;
+	}
+
+	if (!Player->GetSequence() || !Player->GetSequence()->GetMovieScene())
+	{
+		UE_LOG(LogConcertSequencerSync, Warning, TEXT("ApplyPlayRangeToPlayer (%s): Missing sequence or scene"),
+			*Player->GetSequenceName());
+		return;
+	}
+
+	// Convert passed-in range from tick resolution to display rate.
+	UMovieScene* MovieScene = Player->GetSequence()->GetMovieScene();
+	const FFrameRate TickRate = MovieScene->GetTickResolution();
+	const FFrameRate PlayerRate = Player->GetFrameRate();
+
+	const FFrameNumber NewStartFrame = ConvertFrameTime(
+		UE::MovieScene::DiscreteInclusiveLower(PlayRange), TickRate, PlayerRate).FloorToFrame();
+
+	int32 NewDuration;
+	if (bUpperBoundClosed)
+	{
+		const FFrameNumber NewEndFrame = ConvertFrameTime(
+			UE::MovieScene::DiscreteExclusiveUpper(PlayRange), TickRate, PlayerRate).FloorToFrame();
+
+		NewDuration = (NewEndFrame - NewStartFrame).Value;
+	}
+	else
+	{
+		NewDuration = TNumericLimits<int32>::Max() - 1;
+	}
+
+	const int32 CurrentDuration = Player->GetFrameDuration();
+	const FFrameNumber CurrentStartFrame = Player->GetStartTime().Time.GetFrame();
+	if (CurrentDuration != NewDuration || CurrentStartFrame != NewStartFrame)
+	{
+		UE_LOG(LogConcertSequencerSync, Verbose, TEXT("SetFrameRange (%s): Start %d, duration %d (was %d, %d)"),
+			*Player->GetSequenceName(), NewStartFrame.Value, NewDuration, CurrentStartFrame.Value, CurrentDuration);
+
+		Player->SetFrameRange(NewStartFrame.Value, NewDuration);
+	}
+}
+
+bool IsLoopingEnabled(TSharedPtr<ISequencer>& InSequencer)
+{
+	check(InSequencer);
+
+	if (USequencerSettings* Settings = InSequencer->GetSequencerSettings())
+	{
+		ESequencerLoopMode LoopMode = Settings->GetLoopMode();
+		if (LoopMode == SLM_Loop || LoopMode == SLM_LoopSelectionRange)
+		{
+			return true;
+		}
+	}
+
+	return false;
+}
+
+FMovieSceneSequencePlaybackSettings GetPlaybackSettings(bool bInLoopMode)
+{
+	FMovieSceneSequencePlaybackSettings PlaybackSettings;
+
+	// Sequencer behaves differently to Player. Sequencer pauses at the last frame and Player Stops and goes to the
+	// first frame unless we set this flag.
+	//
+	PlaybackSettings.bPauseAtEnd = true;
+
+	if (bInLoopMode && CVarEnableLoopingOnPlayer.GetValueOnAnyThread() > 0)
+	{
+		// Loop indefinitely
+		PlaybackSettings.LoopCount.Value = -1;
+	}
+	return PlaybackSettings;
+}
+
+
+}
+
 void FConcertClientSequencerManager::OnSequencerClosed(TSharedRef<ISequencer> InSequencer)
 {
 	const UMovieSceneSequence* Sequence = InSequencer->GetRootMovieSceneSequence();
@@ -386,6 +479,8 @@
 			SequencerStateEvent.State.PlayerStatus			= (EConcertMovieScenePlayerStatus)Sequencer->GetPlaybackStatus();
 			SequencerStateEvent.State.PlaybackSpeed			= Sequencer->GetPlaybackSpeed();
 
+			SequencerStateEvent.State.bLoopMode = UE::Private::ConcertClientSequencerManager::IsLoopingEnabled(Sequencer);
+
 			UMovieScene* MovieScene = Sequence->GetMovieScene();
 			check(MovieScene);
 			SequencerStateEvent.State.PlaybackRange			= MovieScene->GetPlaybackRange();
@@ -575,10 +670,7 @@
 		return;
 	}
 
-	FMovieSceneSequencePlaybackSettings PlaybackSettings;
-	// Sequencer behaves differently to Player.
-	// Sequencer pauses at the last frame and Player Stops and goes to the first frame unless we set this flag.
-	PlaybackSettings.bPauseAtEnd = true;
+	FMovieSceneSequencePlaybackSettings PlaybackSettings = UE::Private::ConcertClientSequencerManager::GetPlaybackSettings(false);
 
 	// This call will initialize LevelSequenceActor as an output parameter.
 	ALevelSequenceActor* LevelSequenceActor = nullptr;
@@ -834,63 +926,6 @@
 	}
 }
 
-namespace UE::Private::ConcertClientSequencerManager
-{
-
-void ApplyPlayRangeToPlayer(ULevelSequencePlayer* Player, const FFrameNumberRange& PlayRange)
-{
-	check(Player);
-
-	const bool bLowerBoundClosed = PlayRange.GetLowerBound().IsClosed();
-	const bool bUpperBoundClosed = PlayRange.GetUpperBound().IsClosed();
-
-	if (!ensureMsgf(bLowerBoundClosed, TEXT("PlayRange lower bound is open, which is not supported.")))
-	{
-		return;
-	}
-
-	if (!Player->GetSequence() || !Player->GetSequence()->GetMovieScene())
-	{
-		UE_LOG(LogConcertSequencerSync, Warning, TEXT("ApplyPlayRangeToPlayer (%s): Missing sequence or scene"),
-			*Player->GetSequenceName());
-		return;
-	}
-
-	// Convert passed-in range from tick resolution to display rate.
-	UMovieScene* MovieScene = Player->GetSequence()->GetMovieScene();
-	const FFrameRate TickRate = MovieScene->GetTickResolution();
-	const FFrameRate PlayerRate = Player->GetFrameRate();
-
-	const FFrameNumber NewStartFrame = ConvertFrameTime(
-		UE::MovieScene::DiscreteInclusiveLower(PlayRange), TickRate, PlayerRate).FloorToFrame();
-
-	int32 NewDuration;
-	if (bUpperBoundClosed)
-	{
-		const FFrameNumber NewEndFrame = ConvertFrameTime(
-			UE::MovieScene::DiscreteExclusiveUpper(PlayRange), TickRate, PlayerRate).FloorToFrame();
-
-		NewDuration = (NewEndFrame - NewStartFrame).Value;
-	}
-	else
-	{
-		NewDuration = TNumericLimits<int32>::Max() - 1;
-	}
-
-	const int32 CurrentDuration = Player->GetFrameDuration();
-	const FFrameNumber CurrentStartFrame = Player->GetStartTime().Time.GetFrame();
-	if (CurrentDuration != NewDuration || CurrentStartFrame != NewStartFrame)
-	{
-		UE_LOG(LogConcertSequencerSync, Verbose, TEXT("SetFrameRange (%s): Start %d, duration %d (was %d, %d)"),
-			*Player->GetSequenceName(), NewStartFrame.Value, NewDuration, CurrentStartFrame.Value, CurrentDuration);
-
-		Player->SetFrameRange(NewStartFrame.Value, NewDuration);
-	}
-}
-
-}
-
-
 void FConcertClientSequencerManager::ApplyEventToPlayers(const FConcertSequencerState& EventState)
 {
 	UE_LOG(LogConcertSequencerSync, Verbose, TEXT("    Transport: Update sequence player for sequence %s, at frame: %d"), *EventState.SequenceObjectPath, EventState.Time.Time.FrameNumber.Value);
@@ -919,6 +954,8 @@
 		{
 			FFrameTime CurrentTime = Player->GetCurrentTime().Time;
 
+			FMovieSceneSequencePlaybackSettings PlaybackSettings = UE::Private::ConcertClientSequencerManager::GetPlaybackSettings(EventState.bLoopMode);
+			Player->SetPlaybackSettings(PlaybackSettings);
 			// We should be playing back, but are not currently - we compensate the event time for network latency and commence playback
 			if (!Player->IsPlaying())
 			{
diff --git a/Engine/Plugins/Developer/Concert/ConcertSync/ConcertSyncCore/Source/ConcertSyncCore/Public/ConcertSequencerMessages.h b/Engine/Plugins/Developer/Concert/ConcertSync/ConcertSyncCore/Source/ConcertSyncCore/Public/ConcertSequencerMessages.h
--- a/Engine/Plugins/Developer/Concert/ConcertSync/ConcertSyncCore/Source/ConcertSyncCore/Public/ConcertSequencerMessages.h
+++ b/Engine/Plugins/Developer/Concert/ConcertSync/ConcertSyncCore/Source/ConcertSyncCore/Public/ConcertSequencerMessages.h
@@ -50,6 +50,9 @@ struct FConcertSequencerState
 	UPROPERTY()
 	float PlaybackSpeed;
 
+	UPROPERTY()
+	bool bLoopMode = false;
+
 	FConcertSequencerState()
 		: PlayerStatus(EConcertMovieScenePlayerStatus::Stopped)
 		, PlaybackSpeed(1.0f)
