Work around reload issue on sub level by loading the p-level when a sub-level changes.

diff --git a/Engine/Plugins/Developer/Concert/ConcertSync/ConcertSyncClient/Source/ConcertSyncClient/Private/ConcertSyncClientUtil.cpp b/Engine/Plugins/Developer/Concert/ConcertSync/ConcertSyncClient/Source/ConcertSyncClient/Private/ConcertSyncClientUtil.cpp
--- a/Engine/Plugins/Developer/Concert/ConcertSync/ConcertSyncClient/Source/ConcertSyncClient/Private/ConcertSyncClientUtil.cpp
+++ b/Engine/Plugins/Developer/Concert/ConcertSync/ConcertSyncClient/Source/ConcertSyncClient/Private/ConcertSyncClientUtil.cpp
@@ -501,6 +501,24 @@ void SynchronizeAssetRegistry()
 #endif // WITH_EDITOR
 }
 
+bool ShouldReloadPersistentLevel(UPackage *PackageToReload)
+{
+	UWorld* CurrentWorld = ConcertSyncClientUtil::GetCurrentWorld();
+	if (CurrentWorld)
+	{
+		const TArray<ULevel*> Levels = CurrentWorld->GetLevels();
+		for (ULevel* Level : Levels)
+		{
+			if (Level->GetPackage() == PackageToReload)
+			{
+				return true;
+			}
+		}
+	}
+
+	return false;
+}
+
 void HotReloadPackages(TArrayView<const FName> InPackageNames)
 {
 	if (InPackageNames.Num() == 0)
@@ -525,9 +543,12 @@ void HotReloadPackages(TArrayView<const FName> InPackageNames)
 
 	FlushRenderingCommands();
 
+	bool bAddPersistentLevel = false;
+
 	// Find the packages in-memory to content hot-reload
 	TArray<UPackage*> ExistingPackages;
 	ExistingPackages.Reserve(InPackageNames.Num());
+
 	for (const FName& PackageName : InPackageNames)
 	{
 		UPackage* ExistingPackage = FindPackage(nullptr, *PackageName.ToString());
@@ -538,6 +559,20 @@ void HotReloadPackages(TArrayView<const FName> InPackageNames)
 				ExistingPackage->ClearPackageFlags(PKG_NewlyCreated);
 			}
 			ExistingPackages.Add(ExistingPackage);
+			if (ExistingPackage->ContainsMap())
+			{
+				bAddPersistentLevel = ShouldReloadPersistentLevel(ExistingPackage);
+			}
+		}
+	}
+
+	UWorld* CurrentWorld = ConcertSyncClientUtil::GetCurrentWorld();
+	if (CurrentWorld && bAddPersistentLevel)
+	{
+		ULevel* PersistentLevel = CurrentWorld->PersistentLevel;
+		if (PersistentLevel && !ExistingPackages.Contains(PersistentLevel->GetPackage()))
+		{
+			ExistingPackages.Add(PersistentLevel->GetPackage());
 		}
 	}
 
