Work-in-progress receive only support.  Still very experimental.

diff --git a/Engine/Plugins/Developer/Concert/ConcertSync/ConcertSyncClient/Source/ConcertSyncClient/Private/ConcertClientTransactionManager.cpp b/Engine/Plugins/Developer/Concert/ConcertSync/ConcertSyncClient/Source/ConcertSyncClient/Private/ConcertClientTransactionManager.cpp
--- a/Engine/Plugins/Developer/Concert/ConcertSync/ConcertSyncClient/Source/ConcertSyncClient/Private/ConcertClientTransactionManager.cpp
+++ b/Engine/Plugins/Developer/Concert/ConcertSync/ConcertSyncClient/Source/ConcertSyncClient/Private/ConcertClientTransactionManager.cpp
@@ -1,6 +1,7 @@
 // Copyright Epic Games, Inc. All Rights Reserved.
 
 #include "ConcertClientTransactionManager.h"
+#include "ConcertTransactionEvents.h"
 #include "IConcertSession.h"
 #include "ConcertSyncClientLiveSession.h"
 #include "ConcertSyncSessionDatabase.h"
@@ -186,7 +187,10 @@ void FConcertClientTransactionManager::ProcessPending()
 		}
 	}
 
-	SendPendingTransactionEvents();
+	if (CanSendTransactionEvents())
+	{
+		SendPendingTransactionEvents();
+	}
 }
 
 template <typename EventType>
@@ -216,7 +220,7 @@ void FConcertClientTransactionManager::HandleTransactionRejectedEvent(const FCon
 	{
 		GEditor->bSquelchTransactionNotification = true;
 	}
-	
+
 	// if the transaction to undo is the current one, end it.
 	if (GUndo && GUndo->GetContext().TransactionId == InEvent.TransactionId)
 	{
@@ -228,7 +232,7 @@ void FConcertClientTransactionManager::HandleTransactionRejectedEvent(const FCon
 	}
 	// Otherwise undo operations until the requested transaction has been undone.
 	else
-	{		
+	{
 		int32 ReversedQueueIndex = TransBuffer->FindTransactionIndex(InEvent.TransactionId);
 		if (ReversedQueueIndex != INDEX_NONE)
 		{
@@ -270,6 +274,12 @@ bool FConcertClientTransactionManager::CanProcessTransactionEvent() const
 	return TransactionBridge->CanApplyRemoteTransaction() && !bIsSuspended;
 }
 
+bool FConcertClientTransactionManager::CanSendTransactionEvents() const
+{
+	const bool bCanSend = LiveSession->GetSession().GetSendReceiveState() != EConcertSendReceiveState::ReceiveOnly;
+	return bCanSend;
+}
+
 void FConcertClientTransactionManager::ProcessTransactionEvent(const FPendingTransactionToProcessContext& InContext, const FStructOnScope& InEvent)
 {
 	const FConcertTransactionEventBase& TransactionEvent = *(const FConcertTransactionEventBase*)InEvent.GetStructMemory();
@@ -291,12 +301,42 @@ void FConcertClientTransactionManager::ProcessTransactionEvent(const FPendingTra
 #undef PROCESS_OBJECT_UPDATE_EVENT
 }
 
+void FConcertClientTransactionManager::CheckEventForSendConflicts(const FConcertTransactionFinalizedEvent& InEvent)
+{
+	if (CanSendTransactionEvents())
+	{
+		return;
+	}
+
+	auto LookupAndFix = [this](const FConcertObjectId &ObjectId)
+	{
+		if (TSet<FGuid>* NamedPrimary = NameLookupForPendingTransactionsToSend.Find(ObjectId.ObjectName))
+		{
+			UE_LOG(LogConcert, Display, TEXT("ObjectName %s conflicts with pending send transaction."),
+				   *ObjectId.ObjectName.ToString());
+			for (const FGuid Id : *NamedPrimary)
+			{
+				PendingTransactionsToSend.Remove(Id);
+			}
+			NameLookupForPendingTransactionsToSend.Remove(ObjectId.ObjectName);
+		}
+	};
+	LookupAndFix(InEvent.PrimaryObjectId);
+	for (const FConcertExportedObject& ExportedObject : InEvent.ExportedObjects)
+	{
+		LookupAndFix(ExportedObject.ObjectId);
+	}
+	
+}
+
 void FConcertClientTransactionManager::ProcessTransactionFinalizedEvent(const FPendingTransactionToProcessContext& InContext, const FConcertTransactionFinalizedEvent& InEvent)
 {
 	const FConcertSessionVersionInfo* VersionInfo = LiveSession->GetSession().GetSessionInfo().VersionInfos.IsValidIndex(InEvent.VersionIndex) ? &LiveSession->GetSession().GetSessionInfo().VersionInfos[InEvent.VersionIndex] : nullptr;
 	FConcertLocalIdentifierTable LocalIdentifierTable(InEvent.LocalIdentifierState);
 	TransactionBridge->OnApplyTransaction().Broadcast(ETransactionNotification::Begin, /*bIsSnapshot*/ false);
+	CheckEventForSendConflicts(InEvent);
 	TransactionBridge->ApplyRemoteTransaction(InEvent, VersionInfo, InContext.PackagesToProcess, &LocalIdentifierTable, /*bIsSnapshot*/false);
+
 	TransactionBridge->OnApplyTransaction().Broadcast(ETransactionNotification::End, /*bIsSnapshot*/ false);
 }
 
@@ -308,7 +348,10 @@ void FConcertClientTransactionManager::ProcessTransactionSnapshotEvent(const FPe
 	TransactionBridge->OnApplyTransaction().Broadcast(ETransactionNotification::End, /*bIsSnapshot*/ true);
 }
 
-FConcertClientTransactionManager::FPendingTransactionToSend& FConcertClientTransactionManager::HandleLocalTransactionCommon(const FConcertClientLocalTransactionCommonData& InCommonData)
+FConcertClientTransactionManager::FPendingTransactionToSend&
+FConcertClientTransactionManager::HandleLocalTransactionCommon(
+	const FConcertClientLocalTransactionCommonData& InCommonData,
+	const TArray<FConcertExportedObject>& ObjectUpdates)
 {
 	FPendingTransactionToSend* PendingTransactionPtr = PendingTransactionsToSend.Find(InCommonData.OperationId);
 	if (PendingTransactionPtr)
@@ -316,20 +359,71 @@ FConcertClientTransactionManager::FPendingTransactionToSend& FConcertClientTrans
 		PendingTransactionPtr->CommonData = InCommonData;
 		return *PendingTransactionPtr;
 	}
+	return AddPendingToSend(InCommonData, ObjectUpdates);
+}
+
+FConcertClientTransactionManager::FPendingTransactionToSend&
+FConcertClientTransactionManager::AddPendingToSend(
+	const FConcertClientLocalTransactionCommonData& InCommonData,
+	const TArray<FConcertExportedObject>& ObjectUpdates)
+{
 	PendingTransactionsToSendOrder.Add(InCommonData.OperationId);
+	if (!CanSendTransactionEvents())
+	{
+		UObject* PrimaryObject = InCommonData.PrimaryObject.Get();
+		if (PrimaryObject)
+		{
+			FConcertObjectId Id(PrimaryObject);
+			TSet<FGuid>& NamedLookupGuids = NameLookupForPendingTransactionsToSend.FindOrAdd(Id.ObjectName);
+			NamedLookupGuids.Add(InCommonData.OperationId);
+		}
+		for (const FConcertExportedObject& Exported : ObjectUpdates)
+		{
+			TSet<FGuid>& NamedLookupGuids = NameLookupForPendingTransactionsToSend.FindOrAdd(Exported.ObjectId.ObjectName);
+			NamedLookupGuids.Add(InCommonData.OperationId);
+		}
+	}
+	else
+	{
+		NameLookupForPendingTransactionsToSend.Reset();
+	}
 	return PendingTransactionsToSend.Emplace(InCommonData.OperationId, InCommonData);
 }
 
+void FConcertClientTransactionManager::RemovePendingToSend(const FConcertClientLocalTransactionCommonData& InCommonData)
+{
+	PendingTransactionsToSend.Remove(InCommonData.OperationId);
+	if (!CanSendTransactionEvents())
+	{
+		if (UObject* PrimaryObject = InCommonData.PrimaryObject.Get())
+		{
+			FConcertObjectId Id(PrimaryObject);
+			TSet<FGuid>* NamedLookupForObjects = NameLookupForPendingTransactionsToSend.Find(Id.ObjectName);
+			if (NamedLookupForObjects)
+			{
+				NamedLookupForObjects->Remove(InCommonData.OperationId);
+			}
+		}
+	}
+}
+
 void FConcertClientTransactionManager::HandleLocalTransactionSnapshot(const FConcertClientLocalTransactionCommonData& InCommonData, const FConcertClientLocalTransactionSnapshotData& InSnapshotData)
 {
+	if (!CanSendTransactionEvents())
+	{
+		// Don't handle snapshot events when sending to clients.
+		return;
+	}
+
 	if (InCommonData.bIsExcluded)
 	{
-		// Note: We don't remove this from PendingTransactionsToSendOrder as we just skip transactions missing from the map (assuming they've been excluded).
-		PendingTransactionsToSend.Remove(InCommonData.OperationId);
+		// Note: We don't remove this from PendingTransactionsToSendOrder as we just skip transactions missing from the
+		// map (assuming they've been excluded).
+		RemovePendingToSend(InCommonData);
 		return;
 	}
 
-	FPendingTransactionToSend& PendingTransaction = HandleLocalTransactionCommon(InCommonData);
+	FPendingTransactionToSend& PendingTransaction = HandleLocalTransactionCommon(InCommonData, InSnapshotData.SnapshotObjectUpdates);
 	if (PendingTransaction.SnapshotData.SnapshotObjectUpdates.Num() == 0)
 	{
 		PendingTransaction.SnapshotData = InSnapshotData;
@@ -379,7 +473,7 @@ void FConcertClientTransactionManager::HandleLocalTransactionFinalized(const FCo
 	if (InCommonData.bIsExcluded || InFinalizedData.FinalizedObjectUpdates.Num() == 0)
 	{
 		// Note: We don't remove this from PendingTransactionsToSendOrder as we just skip transactions missing from the map (assuming they've been excluded).
-		PendingTransactionsToSend.Remove(InCommonData.OperationId);
+		RemovePendingToSend(InCommonData);
 		return;
 	}
 
@@ -389,13 +483,14 @@ void FConcertClientTransactionManager::HandleLocalTransactionFinalized(const FCo
 		FPendingTransactionToSend* PendingTransactionPtr = PendingTransactionsToSend.Find(InCommonData.OperationId);
 		if (PendingTransactionPtr && PendingTransactionPtr->LastSnapshotTimeSeconds == 0.0)
 		{
-			// Note: We don't remove this from PendingTransactionsToSendOrder as we just skip transactions missing from the map (assuming they've been canceled).
-			PendingTransactionsToSend.Remove(InCommonData.OperationId);
+			// Note: We don't remove this from PendingTransactionsToSendOrder as we just skip transactions missing from
+			// the map (assuming they've been canceled).
+			RemovePendingToSend(InCommonData);
 			return;
 		}
 	}
 
-	FPendingTransactionToSend& PendingTransaction = HandleLocalTransactionCommon(InCommonData);
+	FPendingTransactionToSend& PendingTransaction = HandleLocalTransactionCommon(InCommonData, InFinalizedData.FinalizedObjectUpdates);
 	PendingTransaction.FinalizedData = InFinalizedData;
 	PendingTransaction.bIsFinalized = true;
 }
@@ -461,7 +556,7 @@ void FConcertClientTransactionManager::SendPendingTransactionEvents()
 				}
 				// TODO: Warn about excluded objects?
 
-				PendingTransactionsToSend.Remove(PendingTransactionPtr->CommonData.TransactionId);
+				RemovePendingToSend(PendingTransactionPtr->CommonData);
 				PendingTransactionsToSendOrderIter.RemoveCurrent();
 				continue;
 			}
diff --git a/Engine/Plugins/Developer/Concert/ConcertSync/ConcertSyncClient/Source/ConcertSyncClient/Private/ConcertClientTransactionManager.h b/Engine/Plugins/Developer/Concert/ConcertSync/ConcertSyncClient/Source/ConcertSyncClient/Private/ConcertClientTransactionManager.h
--- a/Engine/Plugins/Developer/Concert/ConcertSync/ConcertSyncClient/Source/ConcertSyncClient/Private/ConcertClientTransactionManager.h
+++ b/Engine/Plugins/Developer/Concert/ConcertSync/ConcertSyncClient/Source/ConcertSyncClient/Private/ConcertClientTransactionManager.h
@@ -122,6 +122,11 @@ private:
 	 */
 	bool CanProcessTransactionEvent() const;
 
+	/**
+	 * Can we send transaction events to remote clients.
+	 */
+	bool CanSendTransactionEvents() const;
+
 	/**
 	 * Process a transaction event.
 	 */
@@ -140,7 +145,7 @@ private:
 	/**
 	 * Handle a local transaction, ensuring that there is a pending entry to be sent.
 	 */
-	FPendingTransactionToSend& HandleLocalTransactionCommon(const FConcertClientLocalTransactionCommonData& InCommonData);
+	FPendingTransactionToSend& HandleLocalTransactionCommon(const FConcertClientLocalTransactionCommonData& InCommonData, const TArray<FConcertExportedObject>& ObjectUpdates);
 
 	/**
 	 * Handle a local transaction being snapshot, queueing it for send later.
@@ -167,6 +172,27 @@ private:
 	 */
 	void SendPendingTransactionEvents();
 
+	/**
+	 * Check the inbound transaction event for conflicts in the pending send. If we are receive only
+	 * and a conflict is detected then we have to remove pending send object.
+	 */
+	void CheckEventForSendConflicts(const FConcertTransactionFinalizedEvent& InEvent);
+
+	/**
+	 * Remove the given pending to send object.
+	 */
+	void RemovePendingToSend(const FConcertClientLocalTransactionCommonData& InCommonData);
+
+	/**
+	 * Remove the given pending to send object.
+	 */
+	FPendingTransactionToSend& AddPendingToSend(const FConcertClientLocalTransactionCommonData& InCommonData, const TArray<FConcertExportedObject>& ObjectUpdates);
+
+	/**
+	   Send any pending
+	 */
+	void IsPrimaryObjectInPendingSend(UObject* PrimaryObject) const;
+
 	/**
 	 * Should process this transaction event?
 	 */
@@ -203,4 +229,10 @@ private:
 	 * Map of transaction IDs to the pending transaction that may be sent in the future (when finalized).
 	 */
 	TMap<FGuid, FPendingTransactionToSend> PendingTransactionsToSend;
+
+	/**
+	 * Named lookup for pending transactions to send.  This is so we can cross check the pending for send against
+	 * remote transactions that may conflict.
+	 */
+	TMap<FName, TSet<FGuid> > NameLookupForPendingTransactionsToSend;
 };
