Respect the console variable set by code when assigning CVar values.

diff --git a/Engine/Plugins/Developer/Concert/ConcertSync/ConcertSyncClient/Source/ConcertSyncClient/Private/ConcertClientSequencerManager.cpp b/Engine/Plugins/Developer/Concert/ConcertSync/ConcertSyncClient/Source/ConcertSyncClient/Private/ConcertClientSequencerManager.cpp
--- a/Engine/Plugins/Developer/Concert/ConcertSync/ConcertSyncClient/Source/ConcertSyncClient/Private/ConcertClientSequencerManager.cpp
+++ b/Engine/Plugins/Developer/Concert/ConcertSync/ConcertSyncClient/Source/ConcertSyncClient/Private/ConcertClientSequencerManager.cpp
@@ -214,6 +214,12 @@ void FConcertClientSequencerManager::Unregister(TSharedRef<IConcertClientSession
 	WeakSession.Reset();
 }
 
+void SetConsoleVariableRespectingPriority(IConsoleVariable* AsVariable, bool bValue)
+{
+	EConsoleVariableFlags Flags = (EConsoleVariableFlags)( AsVariable->GetFlags() & ECVF_SetByMask );
+	AsVariable->Set( bValue ? 1 : 0, Flags );
+}
+
 bool FConcertClientSequencerManager::IsSequencerPlaybackSyncEnabled() const
 {
 	return CVarEnablePlaybackSync.GetValueOnAnyThread() > 0;
@@ -221,7 +227,7 @@ bool FConcertClientSequencerManager::IsSequencerPlaybackSyncEnabled() const
 
 void FConcertClientSequencerManager::SetSequencerPlaybackSync(bool bEnable)
 {
-	CVarEnablePlaybackSync->AsVariable()->Set(bEnable ? 1 : 0);
+	SetConsoleVariableRespectingPriority(CVarEnablePlaybackSync->AsVariable(), bEnable);
 }
 
 bool FConcertClientSequencerManager::IsUnrelatedSequencerTimelineSyncEnabled() const
@@ -231,7 +237,7 @@ bool FConcertClientSequencerManager::IsUnrelatedSequencerTimelineSyncEnabled() c
 
 void FConcertClientSequencerManager::SetUnrelatedSequencerTimelineSync(bool bEnable)
 {
-	CVarEnableUnrelatedTimelineSync->AsVariable()->Set(bEnable ? 1 : 0);
+	SetConsoleVariableRespectingPriority(CVarEnableUnrelatedTimelineSync->AsVariable(), bEnable);
 }
 
 bool FConcertClientSequencerManager::IsSequencerRemoteOpenEnabled() const
@@ -246,12 +252,12 @@ bool FConcertClientSequencerManager::IsSequencerRemoteCloseEnabled() const
 
 void FConcertClientSequencerManager::SetSequencerRemoteOpen(bool bEnable)
 {
-	CVarEnableRemoteSequencerOpen->AsVariable()->Set(bEnable ? 1 : 0);
+	SetConsoleVariableRespectingPriority(CVarEnableRemoteSequencerOpen->AsVariable(), bEnable);
 }
 
 void FConcertClientSequencerManager::SetSequencerRemoteClose(bool bEnable)
 {
-	CVarEnableRemoteSequencerClose->AsVariable()->Set(bEnable ? 1 : 0);
+	SetConsoleVariableRespectingPriority(CVarEnableRemoteSequencerClose->AsVariable(), bEnable);
 }
 
 bool FConcertClientSequencerManager::ShouldAlwaysCloseGameSequencerPlayer() const
