Fix transaction check. We are suppose to check against the incoming transaction and not on the existing transaction buffer.

diff --git a/Engine/Plugins/Developer/Concert/ConcertSync/ConcertSyncClient/Source/ConcertSyncClient/Private/ConcertClientTransactionBridge.cpp b/Engine/Plugins/Developer/Concert/ConcertSync/ConcertSyncClient/Source/ConcertSyncClient/Private/ConcertClientTransactionBridge.cpp
--- a/Engine/Plugins/Developer/Concert/ConcertSync/ConcertSyncClient/Source/ConcertSyncClient/Private/ConcertClientTransactionBridge.cpp
+++ b/Engine/Plugins/Developer/Concert/ConcertSync/ConcertSyncClient/Source/ConcertSyncClient/Private/ConcertClientTransactionBridge.cpp
@@ -178,7 +178,7 @@ struct FEditorTransactionNotification
 	{
 	}
 
-	bool IsTransactionBufferObjectsInSelection() const
+	bool IsTransactedObjectInSelection() const
 	{
 		if (!TransBuffer || !GEditor)
 		{
@@ -188,7 +188,8 @@ struct FEditorTransactionNotification
 		USelection* SelectedActors = GEditor->GetSelectedActors();
 		for (FSelectionIterator Iter(*SelectedActors); Iter; ++Iter)
 		{
-			if (TransBuffer->IsObjectInTransactionBuffer(*Iter))
+			UObject* SelectionObject = *Iter;
+			if (SelectionObject == TransactionContext.PrimaryObject)
 			{
 				return true;
 			}
@@ -204,7 +205,7 @@ struct FEditorTransactionNotification
 			GEditor->bNotifyUndoRedoSelectionChange = true;
 			if (TransBuffer)
 			{
-				if (ConcertSyncClientUtil::IsUserEditing() && !IsTransactionBufferObjectsInSelection())
+				if (ConcertSyncClientUtil::IsUserEditing() && !IsTransactedObjectInSelection())
 				{
 					GEditor->bIgnoreSelectionChange = true;
 					GEditor->bSuspendBroadcastPostUndoRedo = true;
