Ensure components created by Concert are registered

This avoids instanced components being in an unregistered state when we apply updates to them, as that gives them a different post-edit behavior than components created from spawning the actor (which are registered during spawn).

We had a situation where this caused a crash since a component expected its OnUnregister to clear out some state prior to OnRegister running again, but with Concert the state was applied to an unregistered component that skipped the OnUnregister during PostEditUndo, only calling OnRegister.

#rnx

diff --git a/Engine/Plugins/Developer/Concert/ConcertSync/ConcertSyncClient/Source/ConcertSyncClient/Private/ConcertSyncClientUtil.cpp b/Engine/Plugins/Developer/Concert/ConcertSync/ConcertSyncClient/Source/ConcertSyncClient/Private/ConcertSyncClientUtil.cpp
--- a/Engine/Plugins/Developer/Concert/ConcertSync/ConcertSyncClient/Source/ConcertSyncClient/Private/ConcertSyncClientUtil.cpp
+++ b/Engine/Plugins/Developer/Concert/ConcertSync/ConcertSyncClient/Source/ConcertSyncClient/Private/ConcertSyncClientUtil.cpp
@@ -359,6 +359,11 @@ FGetObjectResult GetObject(const FConcertObjectId& InObjectId, const FName InNew
 				else
 				{
 					ObjectResult = FGetObjectResult(NewObject<UObject>(NewObjectOuter, ObjectClass, *ObjectNameToCreate.ToString(), (EObjectFlags)InObjectId.ObjectPersistentFlags), EGetObjectResultFlags::NewlyCreated);
+
+					if (UActorComponent* NewComponent = Cast<UActorComponent>(ObjectResult.Obj))
+					{
+						NewComponent->RegisterComponent();
+					}
 				}
 				
 				// if we have any package assignment, do it here
