Handle conflicts as inbound transactions arrive. This prevents any local modifications from getting messed up from remote changes.

diff --git a/Engine/Plugins/Developer/Concert/ConcertApp/MultiUserClient/Source/MultiUserClient/Private/Widgets/SActiveSession.cpp b/Engine/Plugins/Developer/Concert/ConcertApp/MultiUserClient/Source/MultiUserClient/Private/Widgets/SActiveSession.cpp
--- a/Engine/Plugins/Developer/Concert/ConcertApp/MultiUserClient/Source/MultiUserClient/Private/Widgets/SActiveSession.cpp
+++ b/Engine/Plugins/Developer/Concert/ConcertApp/MultiUserClient/Source/MultiUserClient/Private/Widgets/SActiveSession.cpp
@@ -17,16 +17,20 @@
 #include "EditorFontGlyphs.h"
 #include "FileHelpers.h"
 #include "ISettingsModule.h"
+#include "Misc/AsyncTaskNotification.h"
 #include "Misc/PackageName.h"
 #include "Styling/AppStyle.h"
 #include "Session/History/SSessionHistory.h"
 #include "Session/History/SSessionHistoryWrapper.h"
+#include "Widgets/Notifications/SNotificationList.h"
 #include "Widgets/Text/STextBlock.h"
 #include "Widgets/Images/SImage.h"
 #include "Widgets/Input/SHyperlink.h"
 #include "Widgets/Layout/SExpandableArea.h"
 #include "Widgets/Views/STableViewBase.h"
 #include "Widgets/Views/STableRow.h"
+#include "Framework/Notifications/NotificationManager.h"
+
 
 #define LOCTEXT_NAMESPACE "SActiveSession"
 
@@ -368,6 +372,7 @@ void SActiveSession::Construct(const FArguments& InArgs, TSharedPtr<IConcertSync
 			WeakSessionPtr = ClientSession;
 			ClientInfo = MakeShared<FConcertSessionClientInfo>(FConcertSessionClientInfo{ClientSession->GetSessionClientEndpointId(), ClientSession->GetLocalClientInfo()});
 			ClientSession->OnSessionClientChanged().AddSP(this, &SActiveSession::HandleSessionClientChanged);
+			InConcertSyncClient->GetTransactionBridge()->OnConflictResolutionForPendingSend().AddSP(this, &SActiveSession::OnSendConflict);
 		}
 	}
 
@@ -600,6 +605,10 @@ void SActiveSession::HandleSessionStartup(TSharedRef<IConcertClientSession> InCl
 {
 	WeakSessionPtr = InClientSession;
 	InClientSession->OnSessionClientChanged().AddSP(this, &SActiveSession::HandleSessionClientChanged);
+	if (TSharedPtr<IConcertSyncClient> SyncClient = WeakConcertSyncClient.Pin())
+	{
+		SyncClient->GetTransactionBridge()->OnConflictResolutionForPendingSend().AddSP(this, &SActiveSession::OnSendConflict);
+	}
 
 	ClientInfo = MakeShared<FConcertSessionClientInfo>(FConcertSessionClientInfo{InClientSession->GetSessionClientEndpointId(), InClientSession->GetLocalClientInfo()});
 
@@ -617,6 +626,10 @@ void SActiveSession::HandleSessionShutdown(TSharedRef<IConcertClientSession> InC
 	{
 		WeakSessionPtr.Reset();
 		InClientSession->OnSessionClientChanged().RemoveAll(this);
+		if (TSharedPtr<IConcertSyncClient> SyncClient = WeakConcertSyncClient.Pin())
+		{
+			SyncClient->GetTransactionBridge()->OnConflictResolutionForPendingSend().RemoveAll(this);
+		}
 		Clients.Reset();
 
 		if (ClientsListView.IsValid())
@@ -803,6 +816,18 @@ FSlateFontInfo SActiveSession::GetConnectionIconFontInfo() const
 	return ConnectionIconFontInfo;
 }
 
+void SActiveSession::OnSendConflict(const FConcertConflictDescriptionBase& ConflictDescription)
+{
+	FAsyncTaskNotificationConfig Config;
+	Config.bKeepOpenOnFailure = true;
+	Config.TitleText = ConflictDescription.GetConflictTitle();
+	Config.ExpireDuration = 5.0f;
+
+	FAsyncTaskNotification Notification(Config);
+	Notification.SetHyperlink(FSimpleDelegate(), LOCTEXT("TransactionConflictDetails", "See Details..."));
+	Notification.SetComplete(Config.TitleText, ConflictDescription.GetConflictDetails(), false);
+}
+
 FText SActiveSession::GetConnectionStatusText() const
 {
 	FText StatusText = LOCTEXT("StatusDisconnected", "Disconnected");
diff --git a/Engine/Plugins/Developer/Concert/ConcertApp/MultiUserClient/Source/MultiUserClient/Private/Widgets/SActiveSession.h b/Engine/Plugins/Developer/Concert/ConcertApp/MultiUserClient/Source/MultiUserClient/Private/Widgets/SActiveSession.h
--- a/Engine/Plugins/Developer/Concert/ConcertApp/MultiUserClient/Source/MultiUserClient/Private/Widgets/SActiveSession.h
+++ b/Engine/Plugins/Developer/Concert/ConcertApp/MultiUserClient/Source/MultiUserClient/Private/Widgets/SActiveSession.h
@@ -12,6 +12,7 @@
 
 class IConcertClientSession;
 class IConcertSyncClient;
+struct FConcertConflictDescriptionBase;
 struct FConcertSessionClientInfo;
 class SDockTab;
 class SExpandableArea;
@@ -89,6 +90,9 @@ private:
 	SSplitter::ESizeRule GetHistoryAreaSizeRule() const { return bHistoryAreaExpanded ? SSplitter::ESizeRule::FractionOfParent : SSplitter::ESizeRule::SizeToContent; }
 	void OnHistoryAreaExpansionChanged(bool bExpanded) { bHistoryAreaExpanded = bExpanded; }
 
+	/** Delegate handler when a conflict between inbound transaction and those stored pending outbound transactions. */
+	void OnSendConflict(const FConcertConflictDescriptionBase& ConflictMsg);
+
 private:
 
 	/** Get the text object for the send/receive combo box. */
diff --git a/Engine/Plugins/Developer/Concert/ConcertSync/ConcertSyncClient/Source/ConcertSyncClient/Private/ConcertClientTransactionBridge.cpp b/Engine/Plugins/Developer/Concert/ConcertSync/ConcertSyncClient/Source/ConcertSyncClient/Private/ConcertClientTransactionBridge.cpp
--- a/Engine/Plugins/Developer/Concert/ConcertSync/ConcertSyncClient/Source/ConcertSyncClient/Private/ConcertClientTransactionBridge.cpp
+++ b/Engine/Plugins/Developer/Concert/ConcertSync/ConcertSyncClient/Source/ConcertSyncClient/Private/ConcertClientTransactionBridge.cpp
@@ -651,6 +651,12 @@ FOnApplyTransaction& FConcertClientTransactionBridge::OnApplyTransaction()
 	return OnApplyTransactionDelegate;
 }
 
+FOnConcertConflictResolutionForPendingSend& FConcertClientTransactionBridge::OnConflictResolutionForPendingSend()
+{
+	return OnConflictResolutionForPendingSendDelegate;
+}
+
+
 void FConcertClientTransactionBridge::ApplyRemoteTransaction(const FConcertTransactionEventBase& InEvent, const FConcertSessionVersionInfo* InVersionInfo, const TArray<FName>& InPackagesToProcess, const FConcertLocalIdentifierTable* InLocalIdentifierTablePtr, const bool bIsSnapshot)
 {
 	ConcertClientTransactionBridgeUtil::ProcessTransactionEvent(InEvent, InVersionInfo, InPackagesToProcess, InLocalIdentifierTablePtr, bIsSnapshot, FConcertSyncWorldRemapper(), bIncludeEditorOnlyProperties);
diff --git a/Engine/Plugins/Developer/Concert/ConcertSync/ConcertSyncClient/Source/ConcertSyncClient/Private/ConcertClientTransactionBridge.h b/Engine/Plugins/Developer/Concert/ConcertSync/ConcertSyncClient/Source/ConcertSyncClient/Private/ConcertClientTransactionBridge.h
--- a/Engine/Plugins/Developer/Concert/ConcertSync/ConcertSyncClient/Source/ConcertSyncClient/Private/ConcertClientTransactionBridge.h
+++ b/Engine/Plugins/Developer/Concert/ConcertSync/ConcertSyncClient/Source/ConcertSyncClient/Private/ConcertClientTransactionBridge.h
@@ -20,6 +20,7 @@ public:
 	virtual FOnConcertClientLocalTransactionFinalized& OnLocalTransactionFinalized() override;
 	virtual bool CanApplyRemoteTransaction() const override;
 	virtual FOnApplyTransaction& OnApplyTransaction() override;
+	virtual FOnConcertConflictResolutionForPendingSend& OnConflictResolutionForPendingSend() override;
 
 	virtual void ApplyRemoteTransaction(const FConcertTransactionEventBase& InEvent, const FConcertSessionVersionInfo* InVersionInfo, const TArray<FName>& InPackagesToProcess, const FConcertLocalIdentifierTable* InLocalIdentifierTablePtr, const bool bIsSnapshot) override;
 	virtual void ApplyRemoteTransaction(const FConcertTransactionEventBase& InEvent, const FConcertSessionVersionInfo* InVersionInfo, const TArray<FName>& InPackagesToProcess, const FConcertLocalIdentifierTable* InLocalIdentifierTablePtr, const bool bIsSnapshot, const FConcertSyncWorldRemapper& ConcertSyncWorldRemapper) override;
@@ -72,6 +73,9 @@ private:
 	/** Called when an transaction is finalized */
 	FOnConcertClientLocalTransactionFinalized OnLocalTransactionFinalizedDelegate;
 
+	/** Delegeate called when a inbound transaction conflicts with a pending transaction to send.*/
+	FOnConcertConflictResolutionForPendingSend OnConflictResolutionForPendingSendDelegate;
+
 	/** Called when we are about to apply a transaction. */
 	FOnApplyTransaction OnApplyTransactionDelegate;
 
diff --git a/Engine/Plugins/Developer/Concert/ConcertSync/ConcertSyncClient/Source/ConcertSyncClient/Private/ConcertClientTransactionManager.cpp b/Engine/Plugins/Developer/Concert/ConcertSync/ConcertSyncClient/Source/ConcertSyncClient/Private/ConcertClientTransactionManager.cpp
--- a/Engine/Plugins/Developer/Concert/ConcertSync/ConcertSyncClient/Source/ConcertSyncClient/Private/ConcertClientTransactionManager.cpp
+++ b/Engine/Plugins/Developer/Concert/ConcertSync/ConcertSyncClient/Source/ConcertSyncClient/Private/ConcertClientTransactionManager.cpp
@@ -1,7 +1,10 @@
 // Copyright Epic Games, Inc. All Rights Reserved.
 
 #include "ConcertClientTransactionManager.h"
+#include "Algo/AnyOf.h"
+#include "Components/SceneComponent.h"
 #include "ConcertTransactionEvents.h"
+#include "GameFramework/Actor.h"
 #include "IConcertSession.h"
 #include "ConcertSyncClientLiveSession.h"
 #include "ConcertSyncSessionDatabase.h"
@@ -9,6 +12,8 @@
 #include "ConcertSyncSettings.h"
 #include "ConcertSyncArchives.h"
 #include "ConcertSyncClientUtil.h"
+#include "IConcertSyncClient.h"
+#include "Internationalization/Internationalization.h"
 #include "Scratchpad/ConcertScratchpad.h"
 
 #include "Interfaces/ITargetPlatform.h"
@@ -19,6 +24,7 @@
 #include "Engine/Level.h"
 #include "RenderingThread.h"
 #include "Misc/ScopedSlowTask.h"
+#include "UObject/ObjectMacros.h"
 
 #if WITH_EDITOR
 	#include "Editor.h"
@@ -268,6 +274,11 @@ void FConcertClientTransactionManager::HandleTransactionRejectedEvent(const FCon
 #endif
 }
 
+void FConcertClientTransactionManager::NotifyUserOfSendConflict(const FConcertConflictDescriptionBase& ConflictDescription)
+{
+	TransactionBridge->OnConflictResolutionForPendingSend().Broadcast(ConflictDescription);
+}
+
 bool FConcertClientTransactionManager::CanProcessTransactionEvent() const
 {
 	const bool bIsSuspended = LiveSession->GetSession().GetSendReceiveState() == EConcertSendReceiveState::SendOnly;
@@ -301,6 +312,203 @@ void FConcertClientTransactionManager::ProcessTransactionEvent(const FPendingTra
 #undef PROCESS_OBJECT_UPDATE_EVENT
 }
 
+void FConcertConflictDescriptionAggregate::AddObjectRemoved(const FConcertObjectId& ObjectIdRemoved)
+{
+	UE_LOG(LogConcert, Warning, TEXT("Object %s conflicts with in-bound transaction and has been removed}."),
+		   *ObjectIdRemoved.ObjectName.ToString());
+	ObjectsRemoved.Add(ObjectIdRemoved);
+}
+
+void FConcertConflictDescriptionAggregate::AddObjectRenamed(const FConcertObjectId& OldObjectId, const FConcertObjectId& NewObjectId)
+{
+	if (!ConcertSyncClientUtil::ObjectIdsMatch(OldObjectId, NewObjectId) &&
+		OldObjectId.ObjectName.ToString() != NewObjectId.ObjectName.ToString())
+	{
+		UE_LOG(LogConcert, Warning, TEXT("Object %s conflicts with in-bound transaction and has been renamed to %s."),
+			   *OldObjectId.ObjectName.ToString(), *NewObjectId.ObjectName.ToString());
+		ObjectsRenamed.Add({OldObjectId,NewObjectId});
+	}
+}
+
+FText FConcertConflictDescriptionAggregate::GetConflictTitle() const
+{
+	return FText::Format(
+		LOCTEXT("TransactionConflict", "Inbound transaction conflicts with transactions for pending send. Renamed {0} objects and removed {1} objects."),
+		ObjectsRenamed.Num(), ObjectsRemoved.Num());
+}
+
+FText FConcertConflictDescriptionAggregate::GetConflictDetails() const
+{
+	FTextBuilder RunningTextObject;
+	if (ObjectsRemoved.Num())
+	{
+		RunningTextObject.AppendLine(LOCTEXT("TransactionConflictRemovedObjects", "Removed Objects:"));
+		for (const FConcertObjectId& RemovedObject : ObjectsRemoved)
+		{
+			RunningTextObject.AppendLine(FText::FromString(RemovedObject.ObjectName.ToString()));
+		}
+	}
+
+	if (ObjectsRenamed.Num())
+	{
+		RunningTextObject.AppendLine(LOCTEXT("TransactionConflictRenamedObjects", "Renamed Objects:"));
+		for (const TPair<FConcertObjectId, FConcertObjectId>& Item : ObjectsRenamed)
+		{
+			RunningTextObject.AppendLineFormat(LOCTEXT("ConflictRenamedObject", "{0} -> {1}"),
+											   FText::FromString(Item.Get<0>().ObjectName.ToString()),
+											   FText::FromString(Item.Get<1>().ObjectName.ToString()));
+		}
+	}
+	return RunningTextObject.ToText();
+}
+
+namespace UE::ConcertClientTransactionManager::Private
+{
+/**
+ * Given a UObject to rename we either need to rename this object or the outer for the object. Transactions can refer to
+ * Component.Property modification. However, modifying the component does not address the actual conflicting object.
+ * Instead it is more appropriate to resolve to the parent actor and rename that actor so that it does not globally
+ * conflict anymore once the transaction has been sent to the server.
+ *
+ * @param InObject - UObject to compare.
+ */
+UObject* GetRealObjectToRename(UObject* InObject)
+{
+	auto IsAValidForRename = [](UObject* Obj)
+	{
+		return IsValid(Obj) && (Obj->IsA<AActor>() || Obj->IsA<UActorComponent>());
+	};
+	if (IsAValidForRename(InObject))
+	{
+		if (const UActorComponent* Component = Cast<UActorComponent>(InObject))
+		{
+			if (IsAValidForRename(Component->GetOwner()))
+			{
+				return Component->GetOwner();
+			}
+		}
+		return InObject;
+	}
+	return nullptr;
+}
+
+/**
+ * This structure holds the renaming as applied to an FConcertObjectId so that we can update any tables of the
+ * underlying UObject.
+ */
+struct FRenameObjectResult
+{
+	FConcertObjectId OldObjectId;
+	FConcertObjectId NewObjectId = {};
+
+	FConcertObjectId OldParentObjectId = {};
+	FConcertObjectId NewParentObjectId = {};
+};
+
+/**
+ * Apply the rename to the UObject using unique identifiers to avoid any future conflict.
+ */
+void MakeObjectUniqueToAvoidCollision(UObject* InObject)
+{
+	FString NewName = InObject->GetFName().ToString() + FGuid::NewGuid().ToString(EGuidFormats::Short);
+	InObject->Rename(*NewName, nullptr, REN_ForceNoResetLoaders | REN_NonTransactional);
+}
+
+/**
+ * Given the ObjectId to fixup apply that rename to the resolve UObject and then return the old / new FConcertObjectId
+ * so that any caller can update corresponding data.
+ */
+FRenameObjectResult RenameObjectsToAvoidCollision(const FConcertObjectId& InObjectIdToRename)
+{
+	UObject* ObjectToRename = ConcertSyncClientUtil::GetObject(InObjectIdToRename, FName(), FName(), FName(), false).Obj;
+	// The object does not exist or is marked for GC.
+	if (!IsValid(ObjectToRename))
+	{
+		return {InObjectIdToRename};
+	}
+
+	UObject* TargetObjectToRename = GetRealObjectToRename(ObjectToRename);
+	if (TargetObjectToRename)
+	{
+		FConcertObjectId TargetObjectId(TargetObjectToRename);
+		MakeObjectUniqueToAvoidCollision(TargetObjectToRename);
+		if (ObjectToRename != TargetObjectToRename)
+		{
+			return {InObjectIdToRename, FConcertObjectId(ObjectToRename),
+					TargetObjectId, FConcertObjectId(TargetObjectToRename)};
+		}
+		else
+		{
+			return {InObjectIdToRename, FConcertObjectId(ObjectToRename)};
+		}
+	}
+
+	return {};
+}
+}
+
+void FConcertClientTransactionManager::FixupObjectIdsInPendingSend(
+	const FConcertObjectId& OldObjectId, const FConcertObjectId& NewObjectId, FConcertConflictDescriptionAggregate& ConflictAggregate)
+{
+	/** If we didn't receive a valid OldObjectId then there is nothing to rename. */
+	if (!OldObjectId.IsValid())
+	{
+		return;
+	}
+
+	auto FixNames = [&OldObjectId, &NewObjectId](FPendingTransactionToSend* ToSend)
+	{
+		if (!ToSend)
+		{
+			return;
+		}
+		for (FConcertExportedObject& ExportedObject : ToSend->FinalizedData.FinalizedObjectUpdates)
+		{
+			if (ConcertSyncClientUtil::ObjectIdsMatch(ExportedObject.ObjectId,OldObjectId))
+			{
+				ExportedObject.ObjectId = NewObjectId;
+			}
+		}
+	};
+
+	TSet<FGuid>* TransactionsWithObject = NameLookupForPendingTransactionsToSend.Find(OldObjectId);
+
+	if (!TransactionsWithObject)
+	{
+		return;
+	}
+
+	/*
+	 * If the new object id is not valid then the target object is pending GC or no longer exists. So it should not
+	 * be in our tables anymore.
+	 */
+	const bool bShouldRemove = !NewObjectId.IsValid();
+	if (bShouldRemove)
+	{
+		ConflictAggregate.AddObjectRemoved(OldObjectId);
+		for (const FGuid& Guid : *TransactionsWithObject)
+		{
+			PendingTransactionsToSend.Remove(Guid);
+		}
+	}
+	else
+	{
+		ConflictAggregate.AddObjectRenamed(OldObjectId, NewObjectId);
+		for (const FGuid& Guid : *TransactionsWithObject)
+		{
+			FixNames(PendingTransactionsToSend.Find(Guid));
+		}
+	}
+
+	if (!bShouldRemove)
+	{
+		// Move the renamed FConcertObjectId to a new tracking table entry.
+		NameLookupForPendingTransactionsToSend.Add(NewObjectId, *TransactionsWithObject);
+	}
+	// Remove the old entry named entry.
+	NameLookupForPendingTransactionsToSend.Remove(OldObjectId);
+}
+
 void FConcertClientTransactionManager::CheckEventForSendConflicts(const FConcertTransactionFinalizedEvent& InEvent)
 {
 	if (CanSendTransactionEvents())
@@ -308,25 +516,63 @@ void FConcertClientTransactionManager::CheckEventForSendConflicts(const FConcert
 		return;
 	}
 
-	auto LookupAndFix = [this](const FConcertObjectId &ObjectId)
+	auto IsObjectCreated = [this](const FGuid& Id)
 	{
-		if (TSet<FGuid>* NamedPrimary = NameLookupForPendingTransactionsToSend.Find(ObjectId.ObjectName))
+		if (FPendingTransactionToSend* PendingToSendPtr = PendingTransactionsToSend.Find(Id))
 		{
-			UE_LOG(LogConcert, Display, TEXT("ObjectName %s conflicts with pending send transaction."),
-				   *ObjectId.ObjectName.ToString());
-			for (const FGuid Id : *NamedPrimary)
+			for (const FConcertExportedObject& Exported : PendingToSendPtr->FinalizedData.FinalizedObjectUpdates)
 			{
-				PendingTransactionsToSend.Remove(Id);
+				if (Exported.ObjectData.bAllowCreate)
+				{
+					return true;
+				}
 			}
-			NameLookupForPendingTransactionsToSend.Remove(ObjectId.ObjectName);
 		}
+		return false;
 	};
-	LookupAndFix(InEvent.PrimaryObjectId);
+
+	FConcertConflictDescriptionAggregate ConflictAggregate;
+	auto LookupAndFix = [this, &ConflictAggregate, &IsObjectCreated](const FConcertExportedObject& ExportedObject)
+	{
+		const FConcertObjectId& ObjectId = ExportedObject.ObjectId;
+		if (TSet<FGuid>* NamedPrimary = NameLookupForPendingTransactionsToSend.Find(ObjectId))
+		{
+			// If the object is *new* during our receive only session then we need to go find every instance and
+			// redo the references.
+			if (Algo::AnyOf(*NamedPrimary, IsObjectCreated))
+			{
+				// This exported object is from a create. Therefore the conflict was from a create.
+				// fixup the conflict by renaming the object.
+				//
+				UE::ConcertClientTransactionManager::Private::FRenameObjectResult Result =
+					UE::ConcertClientTransactionManager::Private::RenameObjectsToAvoidCollision(ObjectId);
+				FixupObjectIdsInPendingSend(Result.OldParentObjectId, Result.NewParentObjectId, ConflictAggregate);
+				FixupObjectIdsInPendingSend(Result.OldObjectId, Result.NewObjectId, ConflictAggregate);
+			}
+			else
+			{
+				// Otherwise designate this change has a conflict and must be dropped from the pending list.
+				// TODO UE-170906: We can improve this by using the serialization routines to per-property diff.
+				//
+				ConflictAggregate.AddObjectRemoved(ObjectId);
+				for (const FGuid Id : *NamedPrimary)
+				{
+					PendingTransactionsToSend.Remove(Id);
+				}
+				NameLookupForPendingTransactionsToSend.Remove(ObjectId);
+			}
+		}
+	};
+	// For each exported object check for name conflicts.
+	//
 	for (const FConcertExportedObject& ExportedObject : InEvent.ExportedObjects)
 	{
-		LookupAndFix(ExportedObject.ObjectId);
+		LookupAndFix(ExportedObject);
+	}
+	if (ConflictAggregate.HasAnyResults())
+	{
+		NotifyUserOfSendConflict(ConflictAggregate);
 	}
-	
 }
 
 void FConcertClientTransactionManager::ProcessTransactionFinalizedEvent(const FPendingTransactionToProcessContext& InContext, const FConcertTransactionFinalizedEvent& InEvent)
@@ -374,12 +620,12 @@ FConcertClientTransactionManager::AddPendingToSend(
 		if (PrimaryObject)
 		{
 			FConcertObjectId Id(PrimaryObject);
-			TSet<FGuid>& NamedLookupGuids = NameLookupForPendingTransactionsToSend.FindOrAdd(Id.ObjectName);
+			TSet<FGuid>& NamedLookupGuids = NameLookupForPendingTransactionsToSend.FindOrAdd(Id);
 			NamedLookupGuids.Add(InCommonData.OperationId);
 		}
 		for (const FConcertExportedObject& Exported : ObjectUpdates)
 		{
-			TSet<FGuid>& NamedLookupGuids = NameLookupForPendingTransactionsToSend.FindOrAdd(Exported.ObjectId.ObjectName);
+			TSet<FGuid>& NamedLookupGuids = NameLookupForPendingTransactionsToSend.FindOrAdd(Exported.ObjectId);
 			NamedLookupGuids.Add(InCommonData.OperationId);
 		}
 	}
@@ -398,7 +644,7 @@ void FConcertClientTransactionManager::RemovePendingToSend(const FConcertClientL
 		if (UObject* PrimaryObject = InCommonData.PrimaryObject.Get())
 		{
 			FConcertObjectId Id(PrimaryObject);
-			TSet<FGuid>* NamedLookupForObjects = NameLookupForPendingTransactionsToSend.Find(Id.ObjectName);
+			TSet<FGuid>* NamedLookupForObjects = NameLookupForPendingTransactionsToSend.Find(Id);
 			if (NamedLookupForObjects)
 			{
 				NamedLookupForObjects->Remove(InCommonData.OperationId);
diff --git a/Engine/Plugins/Developer/Concert/ConcertSync/ConcertSyncClient/Source/ConcertSyncClient/Private/ConcertClientTransactionManager.h b/Engine/Plugins/Developer/Concert/ConcertSync/ConcertSyncClient/Source/ConcertSyncClient/Private/ConcertClientTransactionManager.h
--- a/Engine/Plugins/Developer/Concert/ConcertSync/ConcertSyncClient/Source/ConcertSyncClient/Private/ConcertClientTransactionManager.h
+++ b/Engine/Plugins/Developer/Concert/ConcertSync/ConcertSyncClient/Source/ConcertSyncClient/Private/ConcertClientTransactionManager.h
@@ -2,6 +2,7 @@
 
 #pragma once
 
+#include "ConcertTransactionEvents.h"
 #include "CoreMinimal.h"
 #include "Misc/ITransaction.h"
 #include "UObject/Class.h"
@@ -9,9 +10,29 @@
 #include "IConcertSessionHandler.h"
 #include "IConcertClientTransactionBridge.h"
 #include "UObject/StructOnScope.h"
+#include "ConcertSyncClientUtil.h"
 
 class FConcertSyncClientLiveSession;
 
+struct FConcertConflictDescriptionAggregate : public FConcertConflictDescriptionBase
+{
+	virtual ~FConcertConflictDescriptionAggregate() override = default;
+	virtual FText GetConflictDetails() const override;
+	virtual FText GetConflictTitle() const override;
+
+	void AddObjectRemoved(const FConcertObjectId& ObjectIdRemoved);
+	void AddObjectRenamed(const FConcertObjectId& OldObjectId, const FConcertObjectId& NewObjectId);
+
+	bool HasAnyResults() const
+	{
+		return !ObjectsRemoved.IsEmpty() || !ObjectsRenamed.IsEmpty();
+	}
+
+private:
+	TArray<FConcertObjectId> ObjectsRemoved;
+	TArray<TPair<FConcertObjectId, FConcertObjectId>> ObjectsRenamed;
+};
+
 class FConcertClientTransactionManager
 {
 public:
@@ -60,11 +81,25 @@ private:
 	{
 		/** Is this transaction required? */
 		bool bIsRequired = false;
-		
+
 		/** Optional list of packages to process transactions for, or empty to process transactions for all packages */
 		TArray<FName> PackagesToProcess;
 	};
 
+	/** Lookup set key functions. */
+	struct FConcertObjectIdFuncs : public TDefaultMapKeyFuncs<FConcertObjectId, TSet<FGuid>, false>
+	{
+		FORCEINLINE static bool Matches(KeyInitType A, KeyInitType B)
+		{
+			return ConcertSyncClientUtil::ObjectIdsMatch(A, B);
+		}
+
+		FORCEINLINE static uint32 GetKeyHash(KeyInitType InObjectId)
+		{
+			return HashCombine(GetTypeHash(InObjectId.ObjectName), GetTypeHash(InObjectId.ObjectOuterPathName));
+		}
+	};
+
 	/**
 	 * A received pending transaction event that was queued for processing later.
 	 */
@@ -179,19 +214,24 @@ private:
 	void CheckEventForSendConflicts(const FConcertTransactionFinalizedEvent& InEvent);
 
 	/**
-	 * Remove the given pending to send object.
+	 * Fixup any objects from the old ObjectId a given new ObjectId.  The routine will check that the
+	 * given object ids are valid.
+	 *
+	 * @param OldObjectId to search.
+	 * @param NewObjectId to replace.
+	 * @param Aggregate for conflict information.
 	 */
-	void RemovePendingToSend(const FConcertClientLocalTransactionCommonData& InCommonData);
+	void FixupObjectIdsInPendingSend(const FConcertObjectId& OldObjectId, const FConcertObjectId& NewObjectId, FConcertConflictDescriptionAggregate& ConflictAggregate);
 
 	/**
 	 * Remove the given pending to send object.
 	 */
-	FPendingTransactionToSend& AddPendingToSend(const FConcertClientLocalTransactionCommonData& InCommonData, const TArray<FConcertExportedObject>& ObjectUpdates);
+	void RemovePendingToSend(const FConcertClientLocalTransactionCommonData& InCommonData);
 
 	/**
-	   Send any pending
+	 * Remove the given pending to send object.
 	 */
-	void IsPrimaryObjectInPendingSend(UObject* PrimaryObject) const;
+	FPendingTransactionToSend& AddPendingToSend(const FConcertClientLocalTransactionCommonData& InCommonData, const TArray<FConcertExportedObject>& ObjectUpdates);
 
 	/**
 	 * Should process this transaction event?
@@ -203,6 +243,11 @@ private:
 	 */
 	void FillTransactionEvent(const FGuid& InTransactionId, const FGuid& InOperationId, const TArray<FName>& InModifiedPackages, FConcertTransactionEventBase& OutEvent) const;
 
+	/**
+	 * Notify externally that a transaction conflict has occurred.
+	 */
+	void NotifyUserOfSendConflict(const FConcertConflictDescriptionBase& ConflictDescription);
+
 	/**
 	 * Session instance this transaction manager was created for.
 	 */
@@ -234,5 +279,5 @@ private:
 	 * Named lookup for pending transactions to send.  This is so we can cross check the pending for send against
 	 * remote transactions that may conflict.
 	 */
-	TMap<FName, TSet<FGuid> > NameLookupForPendingTransactionsToSend;
+	TMap<FConcertObjectId, TSet<FGuid>, FDefaultSetAllocator, FConcertObjectIdFuncs> NameLookupForPendingTransactionsToSend;
 };
diff --git a/Engine/Plugins/Developer/Concert/ConcertSync/ConcertSyncClient/Source/ConcertSyncClient/Public/IConcertClientTransactionBridge.h b/Engine/Plugins/Developer/Concert/ConcertSync/ConcertSyncClient/Source/ConcertSyncClient/Public/IConcertClientTransactionBridge.h
--- a/Engine/Plugins/Developer/Concert/ConcertSync/ConcertSyncClient/Source/ConcertSyncClient/Public/IConcertClientTransactionBridge.h
+++ b/Engine/Plugins/Developer/Concert/ConcertSync/ConcertSyncClient/Source/ConcertSyncClient/Public/IConcertClientTransactionBridge.h
@@ -58,10 +58,29 @@ enum class ETransactionNotification
 	End,
 };
 
+/**
+ * Data holding information about conflict that occurred while processing an inbound transaction.
+ */
+struct FConcertConflictDescriptionBase
+{
+	FConcertConflictDescriptionBase() = default;
+	virtual ~FConcertConflictDescriptionBase() = default;
+
+	virtual FText GetConflictDetails() const
+	{
+		return FText();
+	};
+	virtual FText GetConflictTitle() const
+	{
+		return FText();
+	};
+};
+
 DECLARE_MULTICAST_DELEGATE_TwoParams(FOnApplyTransaction, ETransactionNotification, const bool bIsSnapshot);
 DECLARE_DELEGATE_RetVal_TwoParams(ETransactionFilterResult, FTransactionFilterDelegate, UObject*, UPackage*);
 DECLARE_MULTICAST_DELEGATE_TwoParams(FOnConcertClientLocalTransactionSnapshot, const FConcertClientLocalTransactionCommonData&, const FConcertClientLocalTransactionSnapshotData&);
 DECLARE_MULTICAST_DELEGATE_TwoParams(FOnConcertClientLocalTransactionFinalized, const FConcertClientLocalTransactionCommonData&, const FConcertClientLocalTransactionFinalizedData&);
+DECLARE_MULTICAST_DELEGATE_OneParam(FOnConcertConflictResolutionForPendingSend, const FConcertConflictDescriptionBase& ConflictDescription);
 
 /**
  * Bridge between the editor transaction system and Concert.
@@ -118,6 +137,9 @@ public:
 	 */
 	virtual FOnConcertClientLocalTransactionFinalized& OnLocalTransactionFinalized() = 0;
 
+	/** Callback when a conflict occurs on pending transactions. */
+	virtual FOnConcertConflictResolutionForPendingSend& OnConflictResolutionForPendingSend() = 0;
+
 	/**
 	 * Can we currently apply a remote transaction event to this local instance?
 	 * @return True if we can apply a remote transaction, false otherwise.
diff --git a/Engine/Plugins/Developer/Concert/ConcertSync/ConcertSyncCore/Source/ConcertSyncCore/Public/ConcertTransactionEvents.h b/Engine/Plugins/Developer/Concert/ConcertSync/ConcertSyncCore/Source/ConcertSyncCore/Public/ConcertTransactionEvents.h
--- a/Engine/Plugins/Developer/Concert/ConcertSync/ConcertSyncCore/Source/ConcertSyncCore/Public/ConcertTransactionEvents.h
+++ b/Engine/Plugins/Developer/Concert/ConcertSync/ConcertSyncCore/Source/ConcertSyncCore/Public/ConcertTransactionEvents.h
@@ -70,6 +70,11 @@ struct FConcertObjectId
 		return FTransactionObjectId(ObjectPackageName, ObjectName, ObjectPathName.ToString(), ObjectOuterPathName, ObjectExternalPackageName, ObjectClassPathName);
 	}
 
+	bool IsValid() const
+	{
+		return !ObjectName.IsNone() && !ObjectOuterPathName.IsNone();
+	}
+
 	UPROPERTY()
 	FName ObjectClassPathName;
 
