Disable collisions for the laser on the presence actor.

diff --git a/Engine/Plugins/Developer/Concert/ConcertSync/ConcertSyncClient/Source/ConcertSyncClient/Private/ConcertClientDesktopPresenceActor.cpp b/Engine/Plugins/Developer/Concert/ConcertSync/ConcertSyncClient/Source/ConcertSyncClient/Private/ConcertClientDesktopPresenceActor.cpp
--- a/Engine/Plugins/Developer/Concert/ConcertSync/ConcertSyncClient/Source/ConcertSyncClient/Private/ConcertClientDesktopPresenceActor.cpp
+++ b/Engine/Plugins/Developer/Concert/ConcertSync/ConcertSyncClient/Source/ConcertSyncClient/Private/ConcertClientDesktopPresenceActor.cpp
@@ -3,6 +3,7 @@
 #include "ConcertClientDesktopPresenceActor.h"
 #include "Components/StaticMeshComponent.h"
 #include "Components/SplineMeshComponent.h"
+#include "Engine/EngineTypes.h"
 #include "IConcertSession.h"
 #include "ConcertPresenceEvents.h"
 #include "ConcertClientPresenceManager.h"
@@ -31,6 +32,8 @@ AConcertClientDesktopPresenceActor::AConcertClientDesktopPresenceActor(const FOb
 	AddOwnedComponent(LaserPointer);
 	LaserPointer->SetMobility(EComponentMobility::Movable);
 	LaserPointer->SetupAttachment(RootComponent);
+	LaserPointer->SetCollisionEnabled(ECollisionEnabled::NoCollision);
+	LaserPointer->SetbNeverNeedsCookedCollisionData(true);
 
 	LastEndPoint = FVector::ZeroVector;
 	bIsLaserVisible = true;
