VR Editor: Refactor multi-user presence interface to use UVREditorModeBase.

diff --git a/Engine/Plugins/Developer/Concert/ConcertSync/ConcertSyncClient/Source/ConcertSyncClient/Private/ConcertClientPresenceManager.cpp b/Engine/Plugins/Developer/Concert/ConcertSync/ConcertSyncClient/Source/ConcertSyncClient/Private/ConcertClientPresenceManager.cpp
--- a/Engine/Plugins/Developer/Concert/ConcertSync/ConcertSyncClient/Source/ConcertSyncClient/Private/ConcertClientPresenceManager.cpp
+++ b/Engine/Plugins/Developer/Concert/ConcertSync/ConcertSyncClient/Source/ConcertSyncClient/Private/ConcertClientPresenceManager.cpp
@@ -911,8 +911,7 @@ FConcertClientDefaultPresenceModeFactory::FConcertClientDefaultPresenceModeFacto
 	IVREditorModule::Get().OnVREditingModeExit().AddRaw(this, &FConcertClientDefaultPresenceModeFactory::OnVREditingModeExit);
 
 	// Set the initial vr device if any
-	UVREditorMode* VRMode = IVREditorModule::Get().GetVRMode();
-	VRDeviceType = VRMode ? VRMode->GetHMDDeviceType() : FName();
+	QueryVRDeviceType();
 }
 
 FConcertClientDefaultPresenceModeFactory::~FConcertClientDefaultPresenceModeFactory()
@@ -946,17 +945,36 @@ void FConcertClientDefaultPresenceModeFactory::AddReferencedObjects(FReferenceCo
 
 void FConcertClientDefaultPresenceModeFactory::OnVREditingModeEnter()
 {
-	UVREditorMode* VRMode = IVREditorModule::Get().GetVRMode();
-	VRDeviceType = VRMode ? VRMode->GetHMDDeviceType() : FName();
+	QueryVRDeviceType();
 	bShouldResetPresenceMode = true;
 }
 
 void FConcertClientDefaultPresenceModeFactory::OnVREditingModeExit()
 {
-	VRDeviceType = FName();
+	VRDeviceType = NAME_None;
 	bShouldResetPresenceMode = true;
 }
 
+void FConcertClientDefaultPresenceModeFactory::QueryVRDeviceType()
+{
+	if (UVREditorModeBase* VRModeBase = IVREditorModule::Get().GetVRModeBase())
+	{
+		if (UVREditorMode* VRMode = Cast<UVREditorMode>(VRModeBase))
+		{
+			VRDeviceType = VRMode->GetHMDDeviceType();
+		}
+		else
+		{
+			static const FName Generic("Generic");
+			VRDeviceType = Generic;
+		}
+	}
+	else
+	{
+		VRDeviceType = NAME_None;
+	}
+}
+
 #else
 
 namespace FConcertClientPresenceManager
diff --git a/Engine/Plugins/Developer/Concert/ConcertSync/ConcertSyncClient/Source/ConcertSyncClient/Private/ConcertClientPresenceManager.h b/Engine/Plugins/Developer/Concert/ConcertSync/ConcertSyncClient/Source/ConcertSyncClient/Private/ConcertClientPresenceManager.h
--- a/Engine/Plugins/Developer/Concert/ConcertSync/ConcertSyncClient/Source/ConcertSyncClient/Private/ConcertClientPresenceManager.h
+++ b/Engine/Plugins/Developer/Concert/ConcertSync/ConcertSyncClient/Source/ConcertSyncClient/Private/ConcertClientPresenceManager.h
@@ -303,6 +303,9 @@ private:
 	/** Handle exiting VR */
 	void OnVREditingModeExit();
 
+	/** Determine new value for VRDeviceType based on current mode */
+	void QueryVRDeviceType();
+
 	/** Holds the manager. */
 	FConcertClientPresenceManager* Manager;
 
diff --git a/Engine/Plugins/Developer/Concert/ConcertSync/ConcertSyncClient/Source/ConcertSyncClient/Private/ConcertClientPresenceMode.cpp b/Engine/Plugins/Developer/Concert/ConcertSync/ConcertSyncClient/Source/ConcertSyncClient/Private/ConcertClientPresenceMode.cpp
--- a/Engine/Plugins/Developer/Concert/ConcertSync/ConcertSyncClient/Source/ConcertSyncClient/Private/ConcertClientPresenceMode.cpp
+++ b/Engine/Plugins/Developer/Concert/ConcertSync/ConcertSyncClient/Source/ConcertSyncClient/Private/ConcertClientPresenceMode.cpp
@@ -95,27 +95,14 @@ FTransform FConcertClientBasePresenceMode::GetTransform()
 			NewHeadTransform = FTransform(PCRotation, PCLocation);
 		}
 	}
+	else if (UVREditorModeBase* VRMode = IVREditorModule::Get().GetVRModeBase())
+	{
+		NewHeadTransform = VRMode->GetHeadTransform();
+	}
 	else if (FLevelEditorViewportClient* PerspectiveViewport = ParentManager->GetPerspectiveViewport())
 	{
-		UViewportWorldInteraction* ViewportWorldInteraction = nullptr;
-		if (UEditorWorldExtensionManager* ExtensionManager = GEditor->GetEditorWorldExtensionsManager())
-		{
-			if (UEditorWorldExtensionCollection* Collection = ExtensionManager->GetEditorWorldExtensions(ParentManager->GetWorld()))
-			{
-				ViewportWorldInteraction = Cast<UViewportWorldInteraction>(Collection->FindExtension(UViewportWorldInteraction::StaticClass()));
-			}
-		}
-
-		// Use the head transform from the VWI if we have one, otherwise fallback to the editor viewport
-		if (ViewportWorldInteraction)
-		{
-			NewHeadTransform = ViewportWorldInteraction->GetHeadTransform();
-		}
-		else
-		{
-			// Take the position from the active viewport
-			NewHeadTransform = FTransform(PerspectiveViewport->GetViewRotation(), PerspectiveViewport->GetViewLocation());
-		}
+		// Take the position from the active viewport
+		NewHeadTransform = FTransform(PerspectiveViewport->GetViewRotation(), PerspectiveViewport->GetViewLocation());
 	}
 	else
 	{ 
@@ -131,7 +118,7 @@ void FConcertClientDesktopPresenceMode::SendEvents(IConcertClientSession& Sessio
 	FConcertClientBasePresenceMode::SendEvents(Session);
 
 	// Send desktop events if not in PIE and not in VR editor
-	if (!ParentManager->IsInPIE() && !IVREditorModule::Get().GetVRMode())
+	if (!ParentManager->IsInPIE() && !IVREditorModule::Get().GetVRModeBase())
 	{
 		FLevelEditorViewportClient* ActiveViewportClient = ParentManager->GetPerspectiveViewport();
 		if (ActiveViewportClient && ActiveViewportClient->Viewport && ActiveViewportClient->Viewport->GetSizeXY().GetMin() > 0)
@@ -223,7 +210,7 @@ void FConcertClientVRPresenceMode::SendEvents(IConcertClientSession& Session)
 	}
 
 	// Send motion controller data if in PIE and VR editor
-	if (ParentManager->IsInPIE() || IVREditorModule::Get().IsVREditorModeActive())
+	if (ParentManager->IsInPIE() || IVREditorModule::Get().GetVRModeBase())
 	{
 		const FTransform PresenceRoomTransform = GetRoomTransform();
 
@@ -266,26 +253,18 @@ void FConcertClientVRPresenceMode::SendEvents(IConcertClientSession& Session)
 		Event.RightMotionControllerPosition = RightMotionControllerTransform.GetLocation();
 
 		// Grab the laser position from VR editor too
-		if (UVREditorMode* VRMode = IVREditorModule::Get().GetVRMode())
+		if (UVREditorModeBase* VRMode = IVREditorModule::Get().GetVRModeBase())
 		{
-			auto HasLaser = [](UVREditorInteractor* Interactor) -> bool
-			{
-				if (Interactor == nullptr)
-				{
-					return false;
-				}
-				AVREditorTeleporter* Teleporter = Interactor->GetTeleportActor();
-				return Interactor->GetControllerType() == EControllerType::AssistingLaser
-					|| Interactor->GetControllerType() == EControllerType::Laser
-					|| (Teleporter && Teleporter->IsAiming());
-			};
-
 			for (int32 HandIndex = 0; HandIndex < 2; ++HandIndex)
 			{
-				UVREditorInteractor* Interactor = HandIndex == 0 ? VRMode->GetHandInteractor(EControllerHand::Left) : VRMode->GetHandInteractor(EControllerHand::Right);
-				if (HasLaser(Interactor))
+				const EControllerHand Hand = (HandIndex == 0)
+					? EControllerHand::Left
+					: EControllerHand::Right;
+
+				FVector LaserStart, LaserEnd;
+				if (VRMode->GetLaserForHand(Hand, LaserStart, LaserEnd))
 				{
-					Event.Lasers[HandIndex] = FConcertLaserData(Interactor->GetLaserStart(), Interactor->GetLaserEnd());
+					Event.Lasers[HandIndex] = FConcertLaserData(LaserStart, LaserEnd);
 				}
 			}
 		}
@@ -312,27 +291,14 @@ FTransform FConcertClientVRPresenceMode::GetRoomTransform()
 			NewRoomTransform = FTransform(PC->GetPawn()->GetActorRotation(), PC->GetPawn()->GetActorLocation());
 		}
 	}
+	else if (UVREditorModeBase* VRMode = IVREditorModule::Get().GetVRModeBase())
+	{
+		NewRoomTransform = VRMode->GetRoomTransform();
+	}
 	else if (FLevelEditorViewportClient* PerspectiveViewport = ParentManager->GetPerspectiveViewport())
 	{
-		UViewportWorldInteraction* ViewportWorldInteraction = nullptr;
-		if (UEditorWorldExtensionManager* ExtensionManager = GEditor->GetEditorWorldExtensionsManager())
-		{
-			if (UEditorWorldExtensionCollection* Collection = ExtensionManager->GetEditorWorldExtensions(ParentManager->GetWorld()))
-			{
-				ViewportWorldInteraction = Cast<UViewportWorldInteraction>(Collection->FindExtension(UViewportWorldInteraction::StaticClass()));
-			}
-		}
-
-		// Use the room transform from the VWI if we have one, otherwise fallback to the editor viewport
-		if (ViewportWorldInteraction)
-		{
-			NewRoomTransform = ViewportWorldInteraction->GetRoomTransform();
-		}
-		else
-		{
-			// Take the position from the active viewport
-			NewRoomTransform = FTransform(PerspectiveViewport->GetViewRotation(), PerspectiveViewport->GetViewLocation());
-		}
+		// Take the position from the active viewport
+		NewRoomTransform = FTransform(PerspectiveViewport->GetViewRotation(), PerspectiveViewport->GetViewLocation());
 	}
 	else
 	{
diff --git a/Engine/Source/Editor/VREditor/Public/VREditorMode.h b/Engine/Source/Editor/VREditor/Public/VREditorMode.h
--- a/Engine/Source/Editor/VREditor/Public/VREditorMode.h
+++ b/Engine/Source/Editor/VREditor/Public/VREditorMode.h
@@ -108,9 +108,11 @@ public:
 	/** Call this to start exiting VR mode */
 	void StartExitingVRMode();
 
+	virtual bool GetLaserForHand(EControllerHand InHand, FVector& OutLaserStart, FVector& OutLaserEnd) const override;
+
 	/** Gets the world space transform of the calibrated VR room origin.  When using a seated VR device, this will feel like the
 	camera's world transform (before any HMD positional or rotation adjustments are applied.) */
-	FTransform GetRoomTransform() const;
+	virtual FTransform GetRoomTransform() const override;
 
 	/** Sets a new transform for the room, in world space.  This is basically setting the editor's camera transform for the viewport */
 	void SetRoomTransform( const FTransform& NewRoomTransform );
@@ -123,7 +125,7 @@ public:
 	 *
 	 * @return	World space space HMD transform
 	 */
-	FTransform GetHeadTransform() const;
+	virtual FTransform GetHeadTransform() const override;
 
 	/** Gets access to the world interaction system (const) */
 	const class UViewportWorldInteraction& GetWorldInteraction() const;
diff --git a/Engine/Source/Editor/VREditor/Public/VREditorModeBase.h b/Engine/Source/Editor/VREditor/Public/VREditorModeBase.h
--- a/Engine/Source/Editor/VREditor/Public/VREditorModeBase.h
+++ b/Engine/Source/Editor/VREditor/Public/VREditorModeBase.h
@@ -40,14 +40,21 @@ public:
 	/** Returns true if the user wants to exit this mode */
 	virtual bool WantsToExitMode() const { return false; }
 
+	/** Returns the transform which takes poses from XR tracking space ("room" space) to world space. */
+	virtual FTransform GetRoomTransform() const { return FTransform::Identity; }
+
+	/** Returns the user's head pose in world space. */
+	virtual FTransform GetHeadTransform() const { return FTransform::Identity; }
+
+	/** Gives start and end points in world space for an active laser, or returns false if there is no active laser for the specified hand. */
+	virtual bool GetLaserForHand(EControllerHand InHand, FVector& OutLaserStart, FVector& OutLaserEnd) const { return false; }
+
 	/** Delegate to be called when async VR mode entry has been completed. */
 	DECLARE_MULTICAST_DELEGATE(FOnVRModeEntryComplete);
 	FOnVRModeEntryComplete& OnVRModeEntryComplete() { return OnVRModeEntryCompleteEvent; }
 
-	// These are slated for deprecation in the next changelist (with accompanying cleanups in external code).
-	//UE_DEPRECATED(5.1, "Use GetVrLevelViewport instead.")
+	// TODO: Deprecate these in favor of replacements below; viewport may be null.
 	const class SLevelViewport& GetLevelViewportPossessedForVR() const;
-	//UE_DEPRECATED(5.1, "Use GetVrLevelViewport instead.")
 	class SLevelViewport& GetLevelViewportPossessedForVR();
 
 	TSharedPtr<SLevelViewport> GetVrLevelViewport() { return VREditorLevelViewportWeakPtr.Pin(); }
diff --git a/Engine/Source/Editor/VREditor/VREditorMode.cpp b/Engine/Source/Editor/VREditor/VREditorMode.cpp
--- a/Engine/Source/Editor/VREditor/VREditorMode.cpp
+++ b/Engine/Source/Editor/VREditor/VREditorMode.cpp
@@ -662,6 +662,28 @@ void UVREditorMode::PostTick( float DeltaTime )
 	bFirstTick = false;
 }
 
+bool UVREditorMode::GetLaserForHand(EControllerHand InHand, FVector& OutLaserStart, FVector& OutLaserEnd) const
+{
+	if (UVREditorInteractor* Interactor = GetHandInteractor(InHand))
+	{
+		AVREditorTeleporter* Teleporter = Interactor->GetTeleportActor();
+
+		const bool bHasLaser =
+			Interactor->GetControllerType() == EControllerType::AssistingLaser
+			|| Interactor->GetControllerType() == EControllerType::Laser
+			|| (Teleporter && Teleporter->IsAiming());
+
+		if (bHasLaser)
+		{
+			OutLaserStart = Interactor->GetLaserStart();
+			OutLaserEnd = Interactor->GetLaserEnd();
+			return true;
+		}
+	}
+
+	return false;
+}
+
 FTransform UVREditorMode::GetRoomTransform() const
 {
 	return WorldInteraction->GetRoomTransform();
