Virtual Scouting: Fix for laser spline meshes recomputing physics derived data every frame. Also fixes an analogous issue in the multi-user VR presence.

diff --git a/Engine/Plugins/Developer/Concert/ConcertSync/ConcertSyncClient/Source/ConcertSyncClient/Private/ConcertClientVRPresenceActor.cpp b/Engine/Plugins/Developer/Concert/ConcertSync/ConcertSyncClient/Source/ConcertSyncClient/Private/ConcertClientVRPresenceActor.cpp
--- a/Engine/Plugins/Developer/Concert/ConcertSync/ConcertSyncClient/Source/ConcertSyncClient/Private/ConcertClientVRPresenceActor.cpp
+++ b/Engine/Plugins/Developer/Concert/ConcertSync/ConcertSyncClient/Source/ConcertSyncClient/Private/ConcertClientVRPresenceActor.cpp
@@ -252,6 +252,7 @@ void AConcertClientVRPresenceActor::InitPresence(const class UConcertAssetContai
 			USplineMeshComponent* SplineSegment = NewObject<USplineMeshComponent>(this);
 			SplineSegment->SetMobility(EComponentMobility::Movable);
 			SplineSegment->SetCollisionEnabled(ECollisionEnabled::NoCollision);
+			SplineSegment->SetbNeverNeedsCookedCollisionData(true);
 			SplineSegment->SetSplineUpDir(FVector::UpVector, false);
 
 			UStaticMesh* StaticMesh = nullptr;
diff --git a/Engine/Source/Editor/VREditor/VREditorInteractor.cpp b/Engine/Source/Editor/VREditor/VREditorInteractor.cpp
--- a/Engine/Source/Editor/VREditor/VREditorInteractor.cpp
+++ b/Engine/Source/Editor/VREditor/VREditorInteractor.cpp
@@ -426,6 +426,7 @@ void UVREditorInteractor::SetupComponent_Implementation( AActor* OwningActor )
 			USplineMeshComponent* SplineSegment = NewObject<USplineMeshComponent>(OwningAvatar);
 			SplineSegment->SetMobility( EComponentMobility::Movable );
 			SplineSegment->SetCollisionEnabled( ECollisionEnabled::NoCollision );
+			SplineSegment->SetbNeverNeedsCookedCollisionData( true );
 			SplineSegment->SetSplineUpDir( FVector::UpVector, false );
 
 			UStaticMesh* StaticMesh = nullptr;
