Listen to editor cleanse events to clear out pending transactions.

diff --git a/Engine/Plugins/Developer/Concert/ConcertSync/ConcertSyncClient/Source/ConcertSyncClient/Private/ConcertClientTransactionBridge.cpp b/Engine/Plugins/Developer/Concert/ConcertSync/ConcertSyncClient/Source/ConcertSyncClient/Private/ConcertClientTransactionBridge.cpp
--- a/Engine/Plugins/Developer/Concert/ConcertSync/ConcertSyncClient/Source/ConcertSyncClient/Private/ConcertClientTransactionBridge.cpp
+++ b/Engine/Plugins/Developer/Concert/ConcertSync/ConcertSyncClient/Source/ConcertSyncClient/Private/ConcertClientTransactionBridge.cpp
@@ -15,6 +15,7 @@
 
 #if WITH_EDITOR
 	#include "Editor.h"
+	#include "EditorSupportDelegates.h"
 	#include "UnrealEdGlobals.h"
 	#include "Editor/UnrealEdEngine.h"
 	#include "Editor/TransBuffer.h"
@@ -610,6 +611,7 @@ FConcertClientTransactionBridge::~FConcertClientTransactionBridge()
 		}
 	}
 	FCoreUObjectDelegates::OnObjectTransacted.RemoveAll(this);
+	FEditorSupportDelegates::CleanseEditor.RemoveAll(this);
 #endif
 
 	FCoreDelegates::OnFEngineLoopInitComplete.RemoveAll(this);
@@ -975,6 +977,7 @@ void FConcertClientTransactionBridge::ConditionalBindUnderlyingLocalTransactionE
 			bHasBoundUnderlyingLocalTransactionEvents = true;
 			TransBuffer->OnTransactionStateChanged().AddRaw(this, &FConcertClientTransactionBridge::HandleTransactionStateChanged);
 			FCoreUObjectDelegates::OnObjectTransacted.AddRaw(this, &FConcertClientTransactionBridge::HandleObjectTransacted);
+			FEditorSupportDelegates::CleanseEditor.AddRaw(this, &FConcertClientTransactionBridge::OnEditorCleanse);
 		}
 	}
 #endif
@@ -985,6 +988,12 @@ void FConcertClientTransactionBridge::OnEngineInitComplete()
 	ConditionalBindUnderlyingLocalTransactionEvents();
 }
 
+void FConcertClientTransactionBridge::OnEditorCleanse()
+{
+	OngoingTransactions.Reset();
+	OngoingTransactionsOrder.Reset();
+}
+
 void FConcertClientTransactionBridge::OnEndFrame()
 {
 	SCOPED_CONCERT_TRACE(FConcertClientTransactionBridge_OnEndFrame);
diff --git a/Engine/Plugins/Developer/Concert/ConcertSync/ConcertSyncClient/Source/ConcertSyncClient/Private/ConcertClientTransactionBridge.h b/Engine/Plugins/Developer/Concert/ConcertSync/ConcertSyncClient/Source/ConcertSyncClient/Private/ConcertClientTransactionBridge.h
--- a/Engine/Plugins/Developer/Concert/ConcertSync/ConcertSyncClient/Source/ConcertSyncClient/Private/ConcertClientTransactionBridge.h
+++ b/Engine/Plugins/Developer/Concert/ConcertSync/ConcertSyncClient/Source/ConcertSyncClient/Private/ConcertClientTransactionBridge.h
@@ -41,6 +41,8 @@ private:
 	/** Called once Engine initialization has finished, to bind to the underlying local transaction events */
 	void OnEngineInitComplete();
 
+	void OnEditorCleanse();
+
 	/** Called at the end of the frame to notify of any transaction updates */
 	void OnEndFrame();
 
