Add a CVar suspending overall post undo redo messages during a transaction event.

diff --git a/Engine/Plugins/Developer/Concert/ConcertSync/ConcertSyncClient/Source/ConcertSyncClient/Private/ConcertClientTransactionBridge.cpp b/Engine/Plugins/Developer/Concert/ConcertSync/ConcertSyncClient/Source/ConcertSyncClient/Private/ConcertClientTransactionBridge.cpp
--- a/Engine/Plugins/Developer/Concert/ConcertSync/ConcertSyncClient/Source/ConcertSyncClient/Private/ConcertClientTransactionBridge.cpp
+++ b/Engine/Plugins/Developer/Concert/ConcertSync/ConcertSyncClient/Source/ConcertSyncClient/Private/ConcertClientTransactionBridge.cpp
@@ -32,6 +32,8 @@ namespace ConcertClientTransactionBridgeUtil
 {
 
 static TAutoConsoleVariable<int32> CVarIgnoreTransactionIncludeFilter(TEXT("Concert.IgnoreTransactionFilters"), 0, TEXT("Ignore Transaction Object Allow List Filtering"));
+static TAutoConsoleVariable<int32> CVarAlwaysSuspendBroadcastUndoRedo(TEXT("Concert.AlwaysSuspendBroadcastUndoRedo"), 0,
+																	  TEXT("Always suspend the undo/redo broadcast message. Can help with transactions that steal user focus."));
 
 bool RunTransactionFilters(const TArray<FTransactionClassFilter>& InFilters, UObject* InObject)
 {
@@ -206,7 +208,7 @@ struct FEditorTransactionNotification
 			GEditor->bNotifyUndoRedoSelectionChange = true;
 			if (TransBuffer)
 			{
-				if (ConcertSyncClientUtil::IsUserEditing() && !IsTransactedObjectInSelection())
+				if (CVarAlwaysSuspendBroadcastUndoRedo.GetValueOnAnyThread() > 0 || (ConcertSyncClientUtil::IsUserEditing() && !IsTransactedObjectInSelection()))
 				{
 					GEditor->bIgnoreSelectionChange = true;
 					GEditor->bSuspendBroadcastPostUndoRedo = true;
