Concert: Account for VR head transform / room transform when using Jump To Presence.

Previously we were setting the user's room transform equal to the target's viewport/head transform. This change effectively sets the head transform instead of the room transform.

diff --git a/Engine/Plugins/Developer/Concert/ConcertSync/ConcertSyncClient/Source/ConcertSyncClient/Private/ConcertClientPresenceManager.cpp b/Engine/Plugins/Developer/Concert/ConcertSync/ConcertSyncClient/Source/ConcertSyncClient/Private/ConcertClientPresenceManager.cpp
--- a/Engine/Plugins/Developer/Concert/ConcertSync/ConcertSyncClient/Source/ConcertSyncClient/Private/ConcertClientPresenceManager.cpp
+++ b/Engine/Plugins/Developer/Concert/ConcertSync/ConcertSyncClient/Source/ConcertSyncClient/Private/ConcertClientPresenceManager.cpp
@@ -806,9 +806,16 @@ void FConcertClientPresenceManager::OnJumpToPresence(FGuid InEndpointId, FTransf
 		}
 		else
 		{
-			FLevelEditorViewportClient* PerspectiveViewport = GetPerspectiveViewport();
-			if (PerspectiveViewport)
+			if (FLevelEditorViewportClient* PerspectiveViewport = GetPerspectiveViewport())
 			{
+				if (UVREditorModeBase* VRMode = IVREditorModule::Get().GetVRModeBase())
+				{
+					const FTransform HeadTransform = VRMode->GetHeadTransform();
+					const FTransform RoomTransform = VRMode->GetRoomTransform();
+					JumpPosition -= HeadTransform.GetLocation() - RoomTransform.GetLocation();
+					OtherClientRotation.Yaw -= (HeadTransform.Rotator().Yaw - RoomTransform.Rotator().Yaw);
+				}
+
 				PerspectiveViewport->SetViewLocation(JumpPosition);
 				PerspectiveViewport->SetViewRotation(OtherClientRotation);
 			}
