Properly zero out pitch and roll when teleporting to a VR participant.

diff --git a/Engine/Plugins/Developer/Concert/ConcertSync/ConcertSyncClient/Source/ConcertSyncClient/Private/ConcertClientPresenceManager.cpp b/Engine/Plugins/Developer/Concert/ConcertSync/ConcertSyncClient/Source/ConcertSyncClient/Private/ConcertClientPresenceManager.cpp
--- a/Engine/Plugins/Developer/Concert/ConcertSync/ConcertSyncClient/Source/ConcertSyncClient/Private/ConcertClientPresenceManager.cpp
+++ b/Engine/Plugins/Developer/Concert/ConcertSync/ConcertSyncClient/Source/ConcertSyncClient/Private/ConcertClientPresenceManager.cpp
@@ -1,6 +1,7 @@
 // Copyright Epic Games, Inc. All Rights Reserved.
 
 #include "ConcertClientPresenceManager.h"
+#include "ConcertPresenceEvents.h"
 #include "CoreMinimal.h"
 #include "IConcertSession.h"
 #include "IConcertSessionHandler.h"
@@ -787,8 +788,8 @@ void FConcertClientPresenceManager::OnJumpToPresence(FGuid InEndpointId, FTransf
 
 		FVector JumpPosition = OtherClientState->Position + InTransformOffset.GetLocation();
 
-		// Disregard pitch and roll when teleporting to a VR presence.
-		if (CurrentAvatarMode && !CurrentAvatarMode->GetVRDeviceType().IsNone())
+		// Disregard pitch and roll when teleporting to a VR presence or we are in VR mode.
+		if (OtherClientState->PresenceType == EPresenceModeType::VRMode || IVREditorModule::Get().GetVRModeBase())
 		{
 			OtherClientRotation.Pitch = 0.0f;
 			OtherClientRotation.Roll = 0.0f;
diff --git a/Engine/Plugins/Developer/Concert/ConcertSync/ConcertSyncClient/Source/ConcertSyncClient/Private/ConcertClientPresenceMode.cpp b/Engine/Plugins/Developer/Concert/ConcertSync/ConcertSyncClient/Source/ConcertSyncClient/Private/ConcertClientPresenceMode.cpp
--- a/Engine/Plugins/Developer/Concert/ConcertSync/ConcertSyncClient/Source/ConcertSyncClient/Private/ConcertClientPresenceMode.cpp
+++ b/Engine/Plugins/Developer/Concert/ConcertSync/ConcertSyncClient/Source/ConcertSyncClient/Private/ConcertClientPresenceMode.cpp
@@ -70,6 +70,10 @@ void FConcertClientBasePresenceMode::SendEvents(IConcertClientSession& Session)
 	PresenceDataUpdatedEvent.WorldPath = *ParentManager->GetPresenceWorldPath(Session.GetSessionClientEndpointId(), EditorPlayModePlaceholder); // The Non-PIE world path, i.e. the "UEDPIE_%d_" decoration stripped away.
 	PresenceDataUpdatedEvent.Position = PresenceHeadTransform.GetLocation();
 	PresenceDataUpdatedEvent.Orientation = PresenceHeadTransform.GetRotation();
+	if (!ParentManager->IsInPIE() && IVREditorModule::Get().GetVRModeBase())
+	{
+		PresenceDataUpdatedEvent.PresenceType = EPresenceModeType::VRMode;
+	}
 
 	SetUpdateIndex(Session, FConcertClientPresenceDataUpdateEvent::StaticStruct()->GetFName(), PresenceDataUpdatedEvent);
 
diff --git a/Engine/Plugins/Developer/Concert/ConcertSync/ConcertSyncCore/Source/ConcertSyncCore/Public/ConcertPresenceEvents.h b/Engine/Plugins/Developer/Concert/ConcertSync/ConcertSyncCore/Source/ConcertSyncCore/Public/ConcertPresenceEvents.h
--- a/Engine/Plugins/Developer/Concert/ConcertSync/ConcertSyncCore/Source/ConcertSyncCore/Public/ConcertPresenceEvents.h
+++ b/Engine/Plugins/Developer/Concert/ConcertSync/ConcertSyncCore/Source/ConcertSyncCore/Public/ConcertPresenceEvents.h
@@ -34,6 +34,13 @@ struct FConcertClientPresenceInVREvent
 	FName VRDevice;
 };
 
+UENUM()
+enum class EPresenceModeType
+{
+	Desktop,
+	VRMode
+};
+
 USTRUCT()
 struct FConcertClientPresenceDataUpdateEvent : public FConcertClientPresenceEventBase
 {
@@ -43,6 +50,7 @@ struct FConcertClientPresenceDataUpdateEvent : public FConcertClientPresenceEven
 		: WorldPath()
 		, Position(FVector::ZeroVector)
 		, Orientation(FQuat::Identity)
+		, PresenceType(EPresenceModeType::Desktop)
 	{
 	}
 
@@ -56,6 +64,9 @@ struct FConcertClientPresenceDataUpdateEvent : public FConcertClientPresenceEven
 
 	UPROPERTY()
 	FQuat Orientation;
+
+	UPROPERTY()
+	EPresenceModeType PresenceType;
 };
 
 USTRUCT()
