Fixed issues applying managed component updates to Concert

The actor we were updating could be reconstructed mid-update which lead to us accessing TRASH components. This happened because the PostEdit of the component would reconstruct its parent actor if the component had previously been registered when its PreEdit call happened.

To avoid this the PreEdit updates need to run actor->component (as the actor PreEdit unregisters its components), while the PostEdit needs to continue to run component->actor. This behavior now matches the editor undo/redo application.

Additionally Concert would always create components as "native" regardless of the CreationMethod provided for the component. This has been corrected so that the correct CreationMethod is applied prior to any PreUpdate calls.

diff --git a/Engine/Plugins/Developer/Concert/ConcertSync/ConcertSyncClient/Source/ConcertSyncClient/Private/ConcertClientTransactionBridge.cpp b/Engine/Plugins/Developer/Concert/ConcertSync/ConcertSyncClient/Source/ConcertSyncClient/Private/ConcertClientTransactionBridge.cpp
--- a/Engine/Plugins/Developer/Concert/ConcertSync/ConcertSyncClient/Source/ConcertSyncClient/Private/ConcertClientTransactionBridge.cpp
+++ b/Engine/Plugins/Developer/Concert/ConcertSync/ConcertSyncClient/Source/ConcertSyncClient/Private/ConcertClientTransactionBridge.cpp
@@ -270,9 +270,9 @@ void ProcessTransactionEvent(const FConcertTransactionEventBase& InEvent, const
 	// Transactions are applied in multiple-phases...
 	//	0) Sort the objects to be processed; creation needs to happen parent->child, and update needs to happen child->parent
 	//	1) Find or create all objects in the transaction (to handle object-interdependencies in the serialized data)
-	//	2) Notify all objects that they are about to be changed (via PreEditUndo)
+	//	2) Notify all objects that they are about to be changed (via PreEditUndo, parent->child)
 	//	3) Update the state of all objects
-	//	4) Notify all objects that they were changed (via PostEditUndo)
+	//	4) Notify all objects that they were changed (via PostEditUndo, child->parent)
 
 	// --------------------------------------------------------------------------------------------------------------------
 	// Phase 0
@@ -326,6 +326,20 @@ void ProcessTransactionEvent(const FConcertTransactionEventBase& InEvent, const
 			{
 				if (TransactionObjectRef.NewlyCreated())
 				{
+					// If this is a new component then we need to apply its CreationMethod (if present) early, as it could affect the PreEdit behavior
+					if (UActorComponent* Component = Cast<UActorComponent>(TransactionObjectRef.Obj))
+					{
+						static const FName NAME_CreationMethod = "CreationMethod";
+						if (const FConcertSerializedPropertyData* CreationMethodPropertyData = ObjectUpdate.PropertyDatas.FindByPredicate([](const FConcertSerializedPropertyData& PropertyData) { return PropertyData.PropertyName == NAME_CreationMethod; }))
+						{
+							if (FProperty* CreationMethodProperty = FindFProperty<FProperty>(TransactionObjectRef.Obj->GetClass(), CreationMethodPropertyData->PropertyName))
+							{
+								FConcertSyncObjectReader ObjectReader(InLocalIdentifierTablePtr, WorldRemapper, InVersionInfo, TransactionObjectRef.Obj, CreationMethodPropertyData->SerializedData);
+								ObjectReader.SerializeProperty(CreationMethodProperty, TransactionObjectRef.Obj);
+							}
+						}
+					}
+
 					// Track this object (and any inner objects, as they must also be new) as newly created
 					NewlyCreatedObjects.Add(TransactionObjectRef.Obj);
 					ForEachObjectWithOuter(TransactionObjectRef.Obj, [&NewlyCreatedObjects](UObject* InnerObj)
@@ -388,7 +402,7 @@ void ProcessTransactionEvent(const FConcertTransactionEventBase& InEvent, const
 #if WITH_EDITOR
 	TArray<TSharedPtr<ITransactionObjectAnnotation>, TInlineAllocator<32>> TransactionAnnotations;
 	TransactionAnnotations.AddDefaulted(SortedExportedObjects.Num());
-	for (int32 ObjectIndex = SortedExportedObjects.Num() - 1; ObjectIndex >= 0; --ObjectIndex)
+	for (int32 ObjectIndex = 0; ObjectIndex < SortedExportedObjects.Num(); ++ObjectIndex)
 	{
 		const ConcertSyncClientUtil::FGetObjectResult& TransactionObjectRef = TransactionObjects[ObjectIndex];
 		const FConcertExportedObject& ObjectUpdate = *SortedExportedObjects[ObjectIndex];
