[Concert Plugin] Log warning instead of crash when receving message to create actor that already exists with same name but different class
#rnx

diff --git a/Engine/Plugins/Developer/Concert/ConcertSync/ConcertSyncClient/Source/ConcertSyncClient/Private/ConcertSyncClientUtil.cpp b/Engine/Plugins/Developer/Concert/ConcertSync/ConcertSyncClient/Source/ConcertSyncClient/Private/ConcertSyncClientUtil.cpp
--- a/Engine/Plugins/Developer/Concert/ConcertSync/ConcertSyncClient/Source/ConcertSyncClient/Private/ConcertSyncClientUtil.cpp
+++ b/Engine/Plugins/Developer/Concert/ConcertSync/ConcertSyncClient/Source/ConcertSyncClient/Private/ConcertSyncClientUtil.cpp
@@ -357,13 +357,22 @@ FGetObjectResult GetObject(const FConcertObjectId& InObjectId, const FName InNew
 
 						if (OwnerWorld)
 						{
-							FActorSpawnParameters SpawnParams;
-							SpawnParams.Name = ObjectNameToCreate;
-							SpawnParams.OverrideLevel = OuterLevel;
-							SpawnParams.SpawnCollisionHandlingOverride = ESpawnActorCollisionHandlingMethod::AlwaysSpawn;
-							SpawnParams.bNoFail = true;
-							SpawnParams.ObjectFlags = (EObjectFlags)InObjectId.ObjectPersistentFlags;
-							ObjectResult = FGetObjectResult(OwnerWorld->SpawnActor<AActor>(ObjectClass, FVector::ZeroVector, FRotator::ZeroRotator, SpawnParams), EGetObjectResultFlags::NewlyCreated);
+							UObject* ExistingObjectOfDifferentClass = StaticFindObjectFast(nullptr, OuterLevel, ObjectNameToCreate);
+							if (!ExistingObjectOfDifferentClass)
+							{
+								FActorSpawnParameters SpawnParams;
+								SpawnParams.Name = ObjectNameToCreate;
+								SpawnParams.OverrideLevel = OuterLevel;
+								SpawnParams.SpawnCollisionHandlingOverride = ESpawnActorCollisionHandlingMethod::AlwaysSpawn;
+								SpawnParams.bNoFail = true;
+								SpawnParams.ObjectFlags = (EObjectFlags)InObjectId.ObjectPersistentFlags;
+								ObjectResult = FGetObjectResult(OwnerWorld->SpawnActor<AActor>(ObjectClass, FVector::ZeroVector, FRotator::ZeroRotator, SpawnParams), EGetObjectResultFlags::NewlyCreated);
+							}
+							else
+							{
+								UE_LOG(LogConcert, Warning, TEXT("Actor '%s' already exists! Expected class: '%s'"), *ExistingObjectOfDifferentClass->GetFullName(), *ObjectClass->GetPathName()); 
+								ensureMsgf(!ExistingObjectOfDifferentClass, TEXT("Actor '%s' already exists! Expected class: '%s'"), *ExistingObjectOfDifferentClass->GetFullName(), *ObjectClass->GetPathName());
+							}
 						}
 						else
 						{
