ConcertSync: Add mechanism for clients to request precaching of level sequence assets.

Also fixes an anti-pattern in FConcertClientSequencerManager of passing a stack-allocated temporary array to AddReferencedObjects, combined with TWeakObjectPtr as a workaround for purges not propagating back to the original SequencePlayers map.

diff --git a/Engine/Plugins/Developer/Concert/ConcertSync/ConcertSyncClient/Source/ConcertSyncClient/Private/ConcertClientSequencerManager.cpp b/Engine/Plugins/Developer/Concert/ConcertSync/ConcertSyncClient/Source/ConcertSyncClient/Private/ConcertClientSequencerManager.cpp
--- a/Engine/Plugins/Developer/Concert/ConcertSync/ConcertSyncClient/Source/ConcertSyncClient/Private/ConcertClientSequencerManager.cpp
+++ b/Engine/Plugins/Developer/Concert/ConcertSync/ConcertSyncClient/Source/ConcertSyncClient/Private/ConcertClientSequencerManager.cpp
@@ -199,10 +199,13 @@ void FConcertClientSequencerManager::Register(TSharedRef<IConcertClientSession>
 	InSession->RegisterCustomEventHandler<FConcertSequencerOpenEvent>(this, &FConcertClientSequencerManager::OnOpenEvent);
 	InSession->RegisterCustomEventHandler<FConcertSequencerStateSyncEvent>(this, &FConcertClientSequencerManager::OnStateSyncEvent);
 	InSession->RegisterCustomEventHandler<FConcertSequencerTimeAdjustmentEvent>(this, &FConcertClientSequencerManager::OnTimeAdjustmentEvent);
+	InSession->RegisterCustomEventHandler<FConcertSequencerPrecacheEvent>(this, &FConcertClientSequencerManager::OnPrecacheEvent);
 }
 
 void FConcertClientSequencerManager::Unregister(TSharedRef<IConcertClientSession> InSession)
 {
+	PrecachedSequences.Empty();
+
 	// Unregister our events and explicitly reset the session ptr
 	if (TSharedPtr<IConcertClientSession> Session = WeakSession.Pin())
 	{
@@ -212,6 +215,7 @@ void FConcertClientSequencerManager::Unregister(TSharedRef<IConcertClientSession
 		Session->UnregisterCustomEventHandler<FConcertSequencerOpenEvent>(this);
 		Session->UnregisterCustomEventHandler<FConcertSequencerStateSyncEvent>(this);
 		Session->UnregisterCustomEventHandler<FConcertSequencerTimeAdjustmentEvent>(this);
+		Session->UnregisterCustomEventHandler<FConcertSequencerPrecacheEvent>(this);
 	}
 	WeakSession.Reset();
 }
@@ -531,12 +535,11 @@ void FConcertClientSequencerManager::HandleAssetReload(const EPackageReloadPhase
 		return;
 	}
 
-	for (TTuple<FName, FSequencePlayer>& Item : SequencePlayers)
+	for (TTuple<FName, ALevelSequenceActor*>& Item : SequencePlayers)
 	{
-		ALevelSequenceActor* LevelSequenceActor = Item.Value.Actor.Get();
+		ALevelSequenceActor* LevelSequenceActor = Item.Value;
 		// If we have a null LevelSequenceActor it means that it has already been destroyed by a close event.  We will not
 		// recreate the asset until all editors have closed it.
-		//
 		if (LevelSequenceActor)
 		{
 			const UPackage* PackageReloaded = InPackageReloadedEvent->GetNewPackage();
@@ -556,14 +559,14 @@ void FConcertClientSequencerManager::ApplyCloseEventToPlayers(const FConcertSequ
 {
 	UE_LOG(LogConcertSequencerSync, Verbose, TEXT("    Closing sequence players: %s, is from controller: %d"), *CloseEvent.SequenceObjectPath, CloseEvent.bControllerClose);
 
-	FSequencePlayer* Player = SequencePlayers.Find(*CloseEvent.SequenceObjectPath);
+	ALevelSequenceActor** Player = SequencePlayers.Find(*CloseEvent.SequenceObjectPath);
 	if (!Player)
 	{
 		UE_LOG(LogConcertSequencerSync, Verbose, TEXT("        No open players to close for sequence %s"), *CloseEvent.SequenceObjectPath);
 		return;
 	}
 
-	ALevelSequenceActor* LevelSequenceActor = Player->Actor.Get();
+	ALevelSequenceActor* LevelSequenceActor = *Player;
 
 	if (CloseEvent.bControllerClose && LevelSequenceActor && LevelSequenceActor->SequencePlayer)
 	{
@@ -575,9 +578,7 @@ void FConcertClientSequencerManager::ApplyCloseEventToPlayers(const FConcertSequ
 		return;
 	}
 
-	DestroyPlayer(*Player, LevelSequenceActor);
-
-	Player->Actor = nullptr;
+	DestroyPlayer(LevelSequenceActor);
 
 	// Always remove the player on close event. This will allow it to be re-opened on an OpenEvent.
 	UE_LOG(LogConcertSequencerSync, VeryVerbose, TEXT("    Removing sequence player for sequence: %s"), *CloseEvent.SequenceObjectPath);
@@ -695,7 +696,7 @@ bool FConcertClientSequencerManager::CanClose(const FConcertSequencerCloseEvent&
 	return InEvent.EditorsWithSequencerOpened == 0 || bShouldClose;
 }
 
-void FConcertClientSequencerManager::DestroyPlayer(FSequencePlayer& Player, ALevelSequenceActor* LevelSequenceActor)
+void FConcertClientSequencerManager::DestroyPlayer(ALevelSequenceActor* LevelSequenceActor)
 {
 	if (LevelSequenceActor && LevelSequenceActor->SequencePlayer)
 	{
@@ -795,14 +796,14 @@ void FConcertClientSequencerManager::ApplyTimeAdjustmentToPlayers(const FConcert
 {
 	UE_LOG(LogConcertSequencerSync, Verbose, TEXT("    TimeAdjustment: Updating sequence players for sequence %s"), *TimeAdjustmentEvent.SequenceObjectPath);
 
-	FSequencePlayer* SeqPlayer = SequencePlayers.Find(*TimeAdjustmentEvent.SequenceObjectPath);
+	ALevelSequenceActor** SeqPlayer = SequencePlayers.Find(*TimeAdjustmentEvent.SequenceObjectPath);
 	if (!SeqPlayer)
 	{
 		UE_LOG(LogConcertSequencerSync, Verbose, TEXT("        No sequence player for sequence %s"), *TimeAdjustmentEvent.SequenceObjectPath);
 		return;
 	}
 
-	ALevelSequenceActor* LevelSequenceActor = SeqPlayer->Actor.Get();
+	ALevelSequenceActor* LevelSequenceActor = *SeqPlayer;
 	if (LevelSequenceActor && LevelSequenceActor->SequencePlayer)
 	{
 		UMovieScene* MovieScene = LevelSequenceActor->LevelSequenceAsset->GetMovieScene();
@@ -813,6 +814,41 @@ void FConcertClientSequencerManager::ApplyTimeAdjustmentToPlayers(const FConcert
 	}
 }
 
+void FConcertClientSequencerManager::OnPrecacheEvent(const FConcertSessionContext& InEventContext, const FConcertSequencerPrecacheEvent& InEvent)
+{
+	for (const FString& SequenceObjectPathStr : InEvent.SequenceObjectPaths)
+	{
+		const FName SequenceObjectPath = *SequenceObjectPathStr;
+		const bool bShouldPrecache = InEvent.bShouldBePrecached;
+		UE_LOG(LogConcertSequencerSync, Verbose, TEXT("OnPrecacheEvent: %s should be %s the precache set"),
+			*SequenceObjectPathStr, bShouldPrecache ? TEXT("added to") : TEXT("removed from"));
+		if (bShouldPrecache)
+		{
+			if (ULevelSequence** ExistingSequence = PrecachedSequences.Find(SequenceObjectPath))
+			{
+				ensure(*ExistingSequence != nullptr);
+				UE_LOG(LogConcertSequencerSync, Warning, TEXT("OnPrecacheEvent: %s already in precache set"), SequenceObjectPath);
+				continue;
+			}
+
+			ULevelSequence* LoadedSequence = LoadObject<ULevelSequence>(nullptr, *SequenceObjectPathStr);
+			if (!LoadedSequence)
+			{
+				UE_LOG(LogConcertSequencerSync, Warning, TEXT("OnPrecacheEvent: Failed to load sequence %s"), SequenceObjectPath);
+				continue;
+			}
+
+			PrecachedSequences.Add(SequenceObjectPath, LoadedSequence);
+		}
+		else
+		{
+			if (!PrecachedSequences.Remove(SequenceObjectPath))
+			{
+				UE_LOG(LogConcertSequencerSync, Warning, TEXT("OnPrecacheEvent: %s not in precache set"), SequenceObjectPath);
+			}
+		}
+	}
+}
 
 void FConcertClientSequencerManager::ApplyTransportEvent(const FConcertSequencerState& EventState)
 {
@@ -938,14 +974,14 @@ void FConcertClientSequencerManager::ApplyEventToPlayers(const FConcertSequencer
 {
 	UE_LOG(LogConcertSequencerSync, Verbose, TEXT("    Transport: Update sequence player for sequence %s, at frame: %d"), *EventState.SequenceObjectPath, EventState.Time.Time.FrameNumber.Value);
 
-	FSequencePlayer* SeqPlayer = SequencePlayers.Find(*EventState.SequenceObjectPath);
+	ALevelSequenceActor** SeqPlayer = SequencePlayers.Find(*EventState.SequenceObjectPath);
 	if (!SeqPlayer)
 	{
 		UE_LOG(LogConcertSequencerSync, Verbose, TEXT("        No sequence player for sequence %s"), *EventState.SequenceObjectPath);
 		return;
 	}
 
-	ALevelSequenceActor* LevelSequenceActor = SeqPlayer->Actor.Get();
+	ALevelSequenceActor* LevelSequenceActor = *SeqPlayer;
 	if (LevelSequenceActor && LevelSequenceActor->SequencePlayer)
 	{
 		ULevelSequencePlayer* Player = LevelSequenceActor->SequencePlayer;
@@ -1056,9 +1092,9 @@ void FConcertClientSequencerManager::OnWorkspaceEndFrameCompleted()
 	{
 		for (const TTuple<FName, FString>& Player : PendingDestroy)
 		{
-			if (FSequencePlayer* SequencePlayer = SequencePlayers.Find(Player.Get<0>()))
+			if (ALevelSequenceActor** SequencePlayer = SequencePlayers.Find(Player.Get<0>()))
 			{
-				DestroyPlayer(*SequencePlayer, SequencePlayer->Actor.Get());
+				DestroyPlayer(*SequencePlayer);
 				SequencePlayers.Remove(Player.Get<0>());
 				PendingCreate.Add(Player.Get<1>());
 			}
@@ -1105,15 +1141,8 @@ void FConcertClientSequencerManager::OnWorkspaceEndFrameCompleted()
 
 void FConcertClientSequencerManager::AddReferencedObjects(FReferenceCollector& Collector)
 {
-	TArray<ALevelSequenceActor*> Actors;
-	for (TTuple<FName, FSequencePlayer>& Item : SequencePlayers)
-	{
-		if (Item.Value.Actor.IsValid())
-		{
-			Actors.Add(Item.Value.Actor.Get());
-		}
-	}
-	Collector.AddReferencedObjects(Actors);
+	Collector.AddReferencedObjects(SequencePlayers);
+	Collector.AddReferencedObjects(PrecachedSequences);
 }
 
 #endif
diff --git a/Engine/Plugins/Developer/Concert/ConcertSync/ConcertSyncClient/Source/ConcertSyncClient/Private/ConcertClientSequencerManager.h b/Engine/Plugins/Developer/Concert/ConcertSync/ConcertSyncClient/Source/ConcertSyncClient/Private/ConcertClientSequencerManager.h
--- a/Engine/Plugins/Developer/Concert/ConcertSync/ConcertSyncClient/Source/ConcertSyncClient/Private/ConcertClientSequencerManager.h
+++ b/Engine/Plugins/Developer/Concert/ConcertSync/ConcertSyncClient/Source/ConcertSyncClient/Private/ConcertClientSequencerManager.h
@@ -19,11 +19,7 @@ class ALevelSequenceActor;
 class ISequencer;
 class FConcertClientWorkspace;
 struct FSequencerInitParams;
-
-struct FSequencePlayer
-{
-	TWeakObjectPtr<ALevelSequenceActor> Actor;
-};
+class ULevelSequence;
 
 /**
  * Sequencer manager that is held by the client sync module that keeps track of open sequencer UIs, regardless of whether a session is open or not
@@ -203,6 +199,13 @@ private:
 	 */
 	void OnTimeAdjustmentEvent(const FConcertSessionContext& InEventContext, const FConcertSequencerTimeAdjustmentEvent& InEvent);
 
+	/**
+	 * Called on receipt of an external sequence precaching event from the server
+	 * @param InEventContext             The session context for the received message.
+	 * @param InEvent                    The sequence precache event.
+	 */
+	void OnPrecacheEvent(const FConcertSessionContext& InEventContext, const FConcertSequencerPrecacheEvent& InEvent);
+
 	/**
 	 * Called when the global time has been changed for the specified Sequencer
 	 *
@@ -320,7 +323,7 @@ private:
 
 private:
 	/** Destroy the given sequence player with corresponding actor. */
-	void DestroyPlayer(FSequencePlayer& Player, ALevelSequenceActor* LevelSequenceActor);
+	void DestroyPlayer(ALevelSequenceActor* LevelSequenceActor);
 
 	/** Indicates if this sequence player can be closed. */
 	bool CanClose(const FConcertSequencerCloseEvent& InEvent) const;
@@ -357,7 +360,10 @@ private:
 	TArray<FOpenSequencerData> OpenSequencers;
 
 	/** Map of opened sequence players, if not in editor mode. */
-	TMap<FName, FSequencePlayer> SequencePlayers;
+	TMap<FName, ALevelSequenceActor*> SequencePlayers;
+
+	/** Map of precached sequence assets. */
+	TMap<FName, ULevelSequence*> PrecachedSequences;
 
 	/** Boolean that is set when we are handling any transport event to prevent re-entrancy */
 	bool bRespondingToTransportEvent;
diff --git a/Engine/Plugins/Developer/Concert/ConcertSync/ConcertSyncCore/Source/ConcertSyncCore/Public/ConcertSequencerMessages.h b/Engine/Plugins/Developer/Concert/ConcertSync/ConcertSyncCore/Source/ConcertSyncCore/Public/ConcertSequencerMessages.h
--- a/Engine/Plugins/Developer/Concert/ConcertSync/ConcertSyncCore/Source/ConcertSyncCore/Public/ConcertSequencerMessages.h
+++ b/Engine/Plugins/Developer/Concert/ConcertSync/ConcertSyncCore/Source/ConcertSyncCore/Public/ConcertSequencerMessages.h
@@ -128,6 +128,31 @@ struct FConcertSequencerTimeAdjustmentEvent
 	FString SequenceObjectPath;
 };
 
+
+/**
+ * Event indicating one or more sequences have been added or removed from the
+ * set of sequences to keep precached for quick dynamic instantiation.
+ *
+ * Can be sent by clients as a request to add or remove their references.
+ *
+ * Can also be received from the server in response to changes to the active set,
+ * or as an initial snapshot of the complete set when joining a session.
+ */
+USTRUCT()
+struct FConcertSequencerPrecacheEvent
+{
+	GENERATED_BODY()
+
+	/** The list of full paths to affected sequences. */
+	UPROPERTY()
+	TArray<FString> SequenceObjectPaths;
+
+	/** True if being added to the precache set, false if being removed. */
+	UPROPERTY()
+	bool bShouldBePrecached = false;
+};
+
+
 #if UE_ENABLE_INCLUDE_ORDER_DEPRECATED_IN_5_2
 #include "CoreMinimal.h"
 #endif
diff --git a/Engine/Plugins/Developer/Concert/ConcertSync/ConcertSyncServer/Source/ConcertSyncServer/Private/ConcertServerSequencerManager.cpp b/Engine/Plugins/Developer/Concert/ConcertSync/ConcertSyncServer/Source/ConcertSyncServer/Private/ConcertServerSequencerManager.cpp
--- a/Engine/Plugins/Developer/Concert/ConcertSync/ConcertSyncServer/Source/ConcertSyncServer/Private/ConcertServerSequencerManager.cpp
+++ b/Engine/Plugins/Developer/Concert/ConcertSync/ConcertSyncServer/Source/ConcertSyncServer/Private/ConcertServerSequencerManager.cpp
@@ -4,6 +4,7 @@
 #include "IConcertSession.h"
 #include "ConcertSyncServerLiveSession.h"
 #include "ConcertWorkspaceMessages.h"
+#include "ConcertLogGlobal.h"
 
 FConcertServerSequencerManager::FConcertServerSequencerManager(const TSharedRef<FConcertSyncServerLiveSession>& InLiveSession)
 {
@@ -28,6 +29,7 @@ void FConcertServerSequencerManager::BindSession(const TSharedRef<FConcertSyncSe
 	LiveSession->GetSession().RegisterCustomEventHandler<FConcertSequencerOpenEvent>(this, &FConcertServerSequencerManager::HandleSequencerOpenEvent);
 	LiveSession->GetSession().RegisterCustomEventHandler<FConcertSequencerTimeAdjustmentEvent>(this, &FConcertServerSequencerManager::HandleSequencerTimeAdjustmentEvent);
 	LiveSession->GetSession().RegisterCustomEventHandler<FConcertWorkspaceSyncAndFinalizeCompletedEvent>(this, &FConcertServerSequencerManager::HandleWorkspaceSyncAndFinalizeCompletedEvent);
+	LiveSession->GetSession().RegisterCustomEventHandler<FConcertSequencerPrecacheEvent>(this, &FConcertServerSequencerManager::HandleSequencerPrecacheEvent);
 }
 
 void FConcertServerSequencerManager::UnbindSession()
@@ -40,6 +42,7 @@ void FConcertServerSequencerManager::UnbindSession()
 		LiveSession->GetSession().UnregisterCustomEventHandler<FConcertSequencerOpenEvent>(this);
 		LiveSession->GetSession().UnregisterCustomEventHandler<FConcertSequencerTimeAdjustmentEvent>(this);
 		LiveSession->GetSession().UnregisterCustomEventHandler<FConcertWorkspaceSyncAndFinalizeCompletedEvent>(this);
+		LiveSession->GetSession().UnregisterCustomEventHandler<FConcertSequencerPrecacheEvent>(this);
 
 		LiveSession.Reset();
 	}
@@ -105,6 +108,56 @@ void FConcertServerSequencerManager::HandleSequencerCloseEvent(const FConcertSes
 	}
 }
 
+void FConcertServerSequencerManager::HandleSequencerPrecacheEvent(const FConcertSessionContext& InEventContext, const FConcertSequencerPrecacheEvent& InEvent)
+{
+	const FGuid& RequestClient = InEventContext.SourceEndpointId;
+	const bool bClientWantsPrecached = InEvent.bShouldBePrecached;
+
+	UE_LOG(LogConcert, Verbose,
+		TEXT("FConcertServerSequencerManager: Precache request from client %s to %s %u sequences"),
+		*RequestClient.ToString(),
+		bClientWantsPrecached ? TEXT("add") : TEXT("remove"),
+		InEvent.SequenceObjectPaths.Num());
+
+	// Represents the net result to broadcast, if any. Contains only sequences
+	// which gained their first, or lost their last, referencer.
+	FConcertSequencerPrecacheEvent OutChanges;
+	OutChanges.bShouldBePrecached = bClientWantsPrecached;
+
+	for (const FString& SequenceObjectPath : InEvent.SequenceObjectPaths)
+	{
+		if (bClientWantsPrecached)
+		{
+			const bool bAddedFirstReferencer = AddSequencePrecacheForClient(RequestClient, SequenceObjectPath);
+			if (bAddedFirstReferencer)
+			{
+				OutChanges.SequenceObjectPaths.Add(SequenceObjectPath);
+			}
+		}
+		else
+		{
+			const bool bRemovedLastReferencer = RemoveSequencePrecacheForClient(RequestClient, SequenceObjectPath);
+			if (bRemovedLastReferencer)
+			{
+				OutChanges.SequenceObjectPaths.Add(SequenceObjectPath);
+			}
+		}
+	}
+
+	if (OutChanges.SequenceObjectPaths.Num() > 0)
+	{
+		for (const FString& SequenceObjectPath : OutChanges.SequenceObjectPaths)
+		{
+			UE_LOG(LogConcert, Verbose,
+				TEXT("FConcertServerSequencerManager: Sequence '%s' %s precache set"),
+				*SequenceObjectPath,
+				bClientWantsPrecached ? TEXT("added to") : TEXT("removed from"));
+		}
+
+		LiveSession->GetSession().SendCustomEvent(OutChanges, LiveSession->GetSession().GetSessionClientEndpointIds(), EConcertMessageFlags::ReliableOrdered | EConcertMessageFlags::UniqueId);
+	}
+}
+
 void FConcertServerSequencerManager::HandleWorkspaceSyncAndFinalizeCompletedEvent(const FConcertSessionContext& InEventContext, const FConcertWorkspaceSyncAndFinalizeCompletedEvent& InEvent)
 {
 	FConcertSequencerStateSyncEvent SequencerStateSyncEvent;
@@ -139,8 +192,116 @@ void FConcertServerSequencerManager::HandleSessionClientChanged(IConcertServerSe
 		}
 	}
 
+	// Newly connected clients need to be sent the current set of precached sequences.
+	// Disconnecting clients need their references removed, which may update other clients.
+	if (InClientStatus == EConcertClientStatus::Connected ||
+		InClientStatus == EConcertClientStatus::Disconnected)
+	{
+		FConcertSequencerPrecacheEvent PrecacheEvent;
+		TArray<FGuid> EventRecipients;
+
+		if (InClientStatus == EConcertClientStatus::Connected)
+		{
+			PrecacheEvent.bShouldBePrecached = true;
+			EventRecipients.Add(InClientInfo.ClientEndpointId);
+		}
+		else
+		{
+			PrecacheEvent.bShouldBePrecached = false;
+			EventRecipients = LiveSession->GetSession().GetSessionClientEndpointIds();
+		}
+
+		for (TMap<FName, FPrecachingState>::TIterator It = PrecacheStates.CreateIterator(); It; ++It)
+		{
+			if (InClientStatus == EConcertClientStatus::Connected)
+			{
+				ensure(It->Value.ReferencingClientEndpoints.Num() > 0);
+				PrecacheEvent.SequenceObjectPaths.Add(It->Key.ToString());
+			}
+			else
+			{
+				if (It->Value.ReferencingClientEndpoints.Remove(InClientInfo.ClientEndpointId))
+				{
+					if (It->Value.ReferencingClientEndpoints.Num() == 0)
+					{
+						PrecacheEvent.SequenceObjectPaths.Add(It->Key.ToString());
+						It.RemoveCurrent();
+					}
+				}
+			}
+		}
+
+		if (PrecacheEvent.SequenceObjectPaths.Num() > 0)
+		{
+			for (const FString& SequenceObjectPath : PrecacheEvent.SequenceObjectPaths)
+			{
+				if (InClientStatus == EConcertClientStatus::Connected)
+				{
+					UE_LOG(LogConcert, Verbose,
+						TEXT("FConcertServerSequencerManager: Client connected; notifying precache set contains sequence '%s'"),
+						*SequenceObjectPath);
+				}
+				else
+				{
+					UE_LOG(LogConcert, Verbose,
+						TEXT("FConcertServerSequencerManager: Client disconnected; last reference to '%s' was released, removed from precache set"),
+						*SequenceObjectPath);
+				}
+			}
+
+			LiveSession->GetSession().SendCustomEvent(PrecacheEvent, EventRecipients, EConcertMessageFlags::ReliableOrdered | EConcertMessageFlags::UniqueId);
+		}
+	}
+
 	// Newly connected clients won't be sent the Sequencer state sync event
 	// until they have synced and finalized their workspace since an open
 	// sequence could have been created by a transaction in the activity
 	// stream.
 }
+
+bool FConcertServerSequencerManager::AddSequencePrecacheForClient(const FGuid& RequestClient, const FString& SequenceObjectPath)
+{
+	bool bAddedFirstReferencer = false;
+
+	FPrecachingState& State = PrecacheStates.FindOrAdd(*SequenceObjectPath);
+	bool bAlreadyInSet = false;
+	State.ReferencingClientEndpoints.Add(RequestClient, &bAlreadyInSet);
+	if (!bAlreadyInSet)
+	{
+		if (State.ReferencingClientEndpoints.Num() == 1)
+		{
+			bAddedFirstReferencer = true;
+		}
+	}
+	else
+	{
+		UE_LOG(LogConcert, Warning, TEXT("FConcertServerSequencerManager: Client %s requested redundant add precache for sequence %s"),
+			*RequestClient.ToString(), *SequenceObjectPath);
+	}
+
+	return bAddedFirstReferencer;
+}
+
+bool FConcertServerSequencerManager::RemoveSequencePrecacheForClient(const FGuid& RequestClient, const FString& SequenceObjectPath)
+{
+	bool bRemovedLastReferencer = false;
+
+	FPrecachingState* MaybeState = PrecacheStates.Find(*SequenceObjectPath);
+	if (MaybeState && MaybeState->ReferencingClientEndpoints.Remove(RequestClient))
+	{
+		if (MaybeState->ReferencingClientEndpoints.Num() == 0)
+		{
+			// Removed last reference.
+			MaybeState = nullptr;
+			PrecacheStates.Remove(*SequenceObjectPath);
+			bRemovedLastReferencer = true;
+		}
+	}
+	else
+	{
+		UE_LOG(LogConcert, Warning, TEXT("FConcertServerSequencerManager: Client %s attempted invalid release precache for sequence %s"),
+			*RequestClient.ToString(), *SequenceObjectPath);
+	}
+
+	return bRemovedLastReferencer;
+}
diff --git a/Engine/Plugins/Developer/Concert/ConcertSync/ConcertSyncServer/Source/ConcertSyncServer/Private/ConcertServerSequencerManager.h b/Engine/Plugins/Developer/Concert/ConcertSync/ConcertSyncServer/Source/ConcertSyncServer/Private/ConcertServerSequencerManager.h
--- a/Engine/Plugins/Developer/Concert/ConcertSync/ConcertSyncServer/Source/ConcertSyncServer/Private/ConcertServerSequencerManager.h
+++ b/Engine/Plugins/Developer/Concert/ConcertSync/ConcertSyncServer/Source/ConcertSyncServer/Private/ConcertServerSequencerManager.h
@@ -36,12 +36,23 @@ private:
 	/** Handler for the sequencer time adjustment event. */
 	void HandleSequencerTimeAdjustmentEvent(const FConcertSessionContext& InEventContext, const FConcertSequencerTimeAdjustmentEvent& InEvent);
 
+	/** Handler for sequence precaching requests. */
+	void HandleSequencerPrecacheEvent(const FConcertSessionContext& InEventContext, const FConcertSequencerPrecacheEvent& InEvent);
+
 	/** Handler for the workspace sync and finalize completed event. */
 	void HandleWorkspaceSyncAndFinalizeCompletedEvent(const FConcertSessionContext& InEventContext, const FConcertWorkspaceSyncAndFinalizeCompletedEvent& InEvent);
 
 	/** Handler for the session clients changed event. */
 	void HandleSessionClientChanged(IConcertServerSession& InSession, EConcertClientStatus InClientStatus, const FConcertSessionClientInfo& InClientInfo);
 
+
+	/** Adds the specified client endpoint as a referencer for the specified sequence, and returns true if that was the first reference for that sequence. */
+	bool AddSequencePrecacheForClient(const FGuid& RequestClient, const FString& SequenceObjectPath);
+
+	/** Adds the specified client endpoint as a referencer for the specified sequence, and returns true if that was the last remaining reference for that sequence. */
+	bool RemoveSequencePrecacheForClient(const FGuid& RequestClient, const FString& SequenceObjectPath);
+
+
 	struct FConcertOpenSequencerState
 	{
 		/** Client Endpoints that have this Sequence opened. */
@@ -56,4 +67,13 @@ private:
 
 	/** Live session tracked by this manager. */
 	TSharedPtr<FConcertSyncServerLiveSession> LiveSession;
+
+	struct FPrecachingState
+	{
+		/** Client endpoints which have requested this sequence be precached. */
+		TSet<FGuid> ReferencingClientEndpoints;
+	};
+
+	/** Map of all level sequences for which precaching has been requested. */
+	TMap<FName, FPrecachingState> PrecacheStates;
 };
