ConcertSync: Fix static analysis warning.

diff --git a/Engine/Plugins/Developer/Concert/ConcertSync/ConcertSyncClient/Source/ConcertSyncClient/Private/ConcertClientSequencerManager.cpp b/Engine/Plugins/Developer/Concert/ConcertSync/ConcertSyncClient/Source/ConcertSyncClient/Private/ConcertClientSequencerManager.cpp
--- a/Engine/Plugins/Developer/Concert/ConcertSync/ConcertSyncClient/Source/ConcertSyncClient/Private/ConcertClientSequencerManager.cpp
+++ b/Engine/Plugins/Developer/Concert/ConcertSync/ConcertSyncClient/Source/ConcertSyncClient/Private/ConcertClientSequencerManager.cpp
@@ -824,9 +824,9 @@ void FConcertClientSequencerManager::OnPrecacheEvent(const FConcertSessionContex
 			*SequenceObjectPathStr, bShouldPrecache ? TEXT("added to") : TEXT("removed from"));
 		if (bShouldPrecache)
 		{
-			if (ULevelSequence** ExistingSequence = PrecachedSequences.Find(SequenceObjectPath))
+			if (ULevelSequence** ExistingSequence = PrecachedSequences.Find(SequenceObjectPath);
+				ExistingSequence && *ExistingSequence)
 			{
-				ensure(*ExistingSequence != nullptr);
 				UE_LOG(LogConcertSequencerSync, Warning, TEXT("OnPrecacheEvent: %s already in precache set"), SequenceObjectPath);
 				continue;
 			}
