Add a multi-user custom event sequencing.  The package event and transaction events exists independently from each other but transaction events can occur between packages events.  We need to make sure that these are sequenced properly otherwise the transaction could get applied incorrectly.

We also fixed ReplayTransaction when loading an asset; replay transactions only need to occur if the transaction was applied after the package event.

diff --git a/Engine/Plugins/Developer/Concert/ConcertMain/Source/Concert/Public/IConcertSession.h b/Engine/Plugins/Developer/Concert/ConcertMain/Source/Concert/Public/IConcertSession.h
--- a/Engine/Plugins/Developer/Concert/ConcertMain/Source/Concert/Public/IConcertSession.h
+++ b/Engine/Plugins/Developer/Concert/ConcertMain/Source/Concert/Public/IConcertSession.h
@@ -2,8 +2,10 @@
 
 #pragma once
 
+#include "ConcertMessageContext.h"
 #include "CoreMinimal.h"
 #include "IConcertEndpoint.h"
+#include "IConcertSession.h"
 #include "IConcertSessionHandler.h"
 #include "ConcertMessages.h"
 #include "Containers/ArrayBuilder.h"
@@ -61,6 +63,107 @@ struct FConcertClientInfoUpdate
 	}
 };
 
+class FConcertSequencedCustomEventManager;
+
+struct FConcertSequencedCustomEvent
+{
+	FConcertSequencedCustomEvent() = default;
+	~FConcertSequencedCustomEvent() = default;
+
+private:
+	friend class FConcertSequencedCustomEventManager;
+
+	FConcertSequencedCustomEvent(int32 InIndex)
+		: Index(InIndex)
+	{
+	}
+
+	int32 Index = INDEX_NONE;
+};
+
+/** Allows users to reserve a spot in the outbound queues. */
+class FConcertSequencedCustomEventManager
+{
+public:
+	FConcertSequencedCustomEventManager()
+	{
+	}
+	~FConcertSequencedCustomEventManager()
+	{
+	}
+
+	FConcertSequencedCustomEvent AddSequencedCustomEvent()
+	{
+		Events.Add(CurrentId, {});
+		return FConcertSequencedCustomEvent(CurrentId++);
+	}
+
+	void RemoveSequencedEvent(const FConcertSequencedCustomEvent& EventDetails)
+	{
+		if (EventDetails.Index == INDEX_NONE)
+		{
+			return;
+		}
+		Events.Remove(EventDetails.Index);
+	}
+
+	bool HasEventsPending() const
+	{
+		return !Events.IsEmpty();
+	}
+
+	struct FPendingCustomEvent
+	{
+		FConcertSession_CustomEvent CustomEvent;
+		EConcertMessageFlags Flags;
+		TMap<FName, FString> Annotations;
+		FGuid ServerEndpointId;
+	};
+
+	TOptional<FPendingCustomEvent> PopSendEvent()
+	{
+		TOptional<FPendingCustomEvent> ReturnVal;
+		if (LastSendId < CurrentId)
+		{
+			TOptional<FPendingCustomEvent> *CustomEvent = Events.Find(LastSendId);
+			// Find the first valid value we have stored.
+			while (!CustomEvent)
+			{
+				LastSendId++;
+				CustomEvent = Events.Find(LastSendId);
+			}
+			if (*CustomEvent)
+			{
+				// If the optional has a value then it is available for sending.
+				ReturnVal = *CustomEvent;
+				Events.Remove(LastSendId);
+				LastSendId++;
+			}
+		}
+		return ReturnVal;
+	}
+
+	void FillPendingSequenceEvent(const FConcertSequencedCustomEvent& InEventId, FPendingCustomEvent InEventData)
+	{
+		TOptional<FPendingCustomEvent>* PendingEvent = Events.Find(InEventId.Index);
+		// The event data should exist and it should not have any data in the optional.
+		check(PendingEvent);
+		check(!(*PendingEvent));
+		*PendingEvent = MoveTemp(InEventData);
+	}
+
+	bool CanSend(const FConcertSequencedCustomEvent& InEventId) const
+	{
+		return InEventId.Index == LastSendId;
+	}
+
+private:
+	TMap<int32, TOptional<FPendingCustomEvent> > Events;
+
+	int32 LastSendId = 0;
+	int32 CurrentId = 0;
+};
+
 /** Interface for Concert sessions */
 class IConcertSession
 {
@@ -149,18 +252,18 @@ public:
 	 * Send a custom event event to the given endpoint
 	 */
 	template<typename EventType>
-	void SendCustomEvent(const EventType& Event, const FGuid& DestinationEndpointId, EConcertMessageFlags Flags)
+	void SendCustomEvent(const EventType& Event, const FGuid& DestinationEndpointId, EConcertMessageFlags Flags, TOptional<FConcertSequencedCustomEvent> SequencedId = {})
 	{
-		InternalSendCustomEvent(EventType::StaticStruct(), &Event, TArrayBuilder<FGuid>().Add(DestinationEndpointId), Flags);
+		InternalSendCustomEvent(EventType::StaticStruct(), &Event, TArrayBuilder<FGuid>().Add(DestinationEndpointId), Flags, SequencedId);
 	}
 
 	/**
 	 * Send a custom event event to the given endpoints
 	 */
 	template<typename EventType>
-	void SendCustomEvent(const EventType& Event, const TArray<FGuid>& DestinationEndpointIds, EConcertMessageFlags Flags)
+	void SendCustomEvent(const EventType& Event, const TArray<FGuid>& DestinationEndpointIds, EConcertMessageFlags Flags, TOptional<FConcertSequencedCustomEvent> SequencedId = {})
 	{
-		InternalSendCustomEvent(EventType::StaticStruct(), &Event, DestinationEndpointIds, Flags);
+		InternalSendCustomEvent(EventType::StaticStruct(), &Event, DestinationEndpointIds, Flags, SequencedId);
 	}
 
 	/**
@@ -221,7 +324,7 @@ protected:
 	/**
 	 * Send a custom event event to the given endpoints
 	 */
-	virtual void InternalSendCustomEvent(const UScriptStruct* EventType, const void* EventData, const TArray<FGuid>& DestinationEndpointIds, EConcertMessageFlags Flags) = 0;
+	virtual void InternalSendCustomEvent(const UScriptStruct* EventType, const void* EventData, const TArray<FGuid>& DestinationEndpointIds, EConcertMessageFlags Flags, TOptional<FConcertSequencedCustomEvent> InSequenceId={}) = 0;
 
 	/**
 	 * Register a custom request handler for this session
@@ -237,6 +340,7 @@ protected:
 	 * Send a custom request to the given endpoint
 	 */
 	virtual void InternalSendCustomRequest(const UScriptStruct* RequestType, const void* RequestData, const FGuid& DestinationEndpointId, const TSharedRef<IConcertSessionCustomResponseHandler>& Handler) = 0;
+
 };
 
 /** Interface for Concert server sessions */
@@ -256,7 +360,7 @@ public:
 
 	/** Callback when a session client state changes */
 	virtual FOnConcertServerSessionClientChanged& OnSessionClientChanged() = 0;
-	
+
 	/** Callback when a session message is acknowledged */
 	virtual FOnConcertMessageAcknowledgementReceivedFromLocalEndpoint& OnConcertMessageAcknowledgementReceived() = 0;
 };
@@ -309,6 +413,9 @@ public:
 		return GetSendReceiveState() == EConcertSendReceiveState::SendOnly;
 	}
 
+	/** Manager for helping sequence custom events on client sessions. */
+	virtual FConcertSequencedCustomEventManager& GetSequencedEventManager() = 0;
+
 	/** Get the send/receive state for this session. */
 	virtual EConcertSendReceiveState GetSendReceiveState() const = 0;
 
diff --git a/Engine/Plugins/Developer/Concert/ConcertMain/Source/ConcertClient/Private/ConcertClientSession.cpp b/Engine/Plugins/Developer/Concert/ConcertMain/Source/ConcertClient/Private/ConcertClientSession.cpp
--- a/Engine/Plugins/Developer/Concert/ConcertMain/Source/ConcertClient/Private/ConcertClientSession.cpp
+++ b/Engine/Plugins/Developer/Concert/ConcertMain/Source/ConcertClient/Private/ConcertClientSession.cpp
@@ -208,14 +208,26 @@ FString FConcertClientSession::GetSessionWorkingDirectory() const
 	return SessionDirectory;
 }
 
-void FConcertClientSession::InternalSendCustomEvent(const UScriptStruct* EventType, const void* EventData, const TArray<FGuid>& DestinationEndpointIds, EConcertMessageFlags Flags)
+void FConcertClientSession::SendNextCustomEvent(const FConcertSequencedCustomEventManager::FPendingCustomEvent& PendingEvent)
+{
+	ClientSessionEndpoint->SendEvent(PendingEvent.CustomEvent, PendingEvent.ServerEndpointId, PendingEvent.Flags, PendingEvent.Annotations);
+}
+
+void FConcertClientSession::InternalSendCustomEvent(const UScriptStruct* EventType, const void* EventData, const TArray<FGuid>& DestinationEndpointIds, EConcertMessageFlags Flags, TOptional<FConcertSequencedCustomEvent> InSequenceId)
 {
 	if (DestinationEndpointIds.Num() == 0)
 	{
+		if (InSequenceId)
+		{
+			CustomEventSequenceManager.RemoveSequencedEvent(*InSequenceId);
+		}
 		return;
 	}
 
-	// TODO: don't send if not connected
+	if (ConnectionStatus == EConcertConnectionStatus::Disconnected)
+	{
+		return;
+	}
 
 	// Build the event
 	FConcertSession_CustomEvent CustomEvent;
@@ -228,8 +240,31 @@ void FConcertClientSession::InternalSendCustomEvent(const UScriptStruct* EventTy
 		Annotations.Add(ConcertClientMessageIdName, FGuid::NewGuid().ToString());
 	}
 
-	// Send the event
-	ClientSessionEndpoint->SendEvent(CustomEvent, SessionInfo.ServerEndpointId, Flags, Annotations);
+	using FPendingEvent = FConcertSequencedCustomEventManager::FPendingCustomEvent;
+	FPendingEvent InEvent{
+		MoveTemp(CustomEvent),
+		Flags,
+		MoveTemp(Annotations),
+		SessionInfo.ServerEndpointId
+	};
+
+	const bool bNotReliablyOrdered = !EnumHasAnyFlags(Flags, EConcertMessageFlags::ReliableOrdered);
+	const bool bHasCustomEventsPending = CustomEventSequenceManager.HasEventsPending();
+
+	if (!bHasCustomEventsPending || bNotReliablyOrdered)
+	{
+		// We are not currently sequencing the events so send it immediately.
+		SendNextCustomEvent(InEvent);
+		return;
+	}
+	// If the user provided a id then use that id to store the value in the manager; otherwise create a new entry.
+	FConcertSequencedCustomEvent EventId = InSequenceId ? *InSequenceId : CustomEventSequenceManager.AddSequencedCustomEvent();
+	CustomEventSequenceManager.FillPendingSequenceEvent(EventId, MoveTemp(InEvent));
+
+	while (TOptional<FPendingEvent> NextToSend = CustomEventSequenceManager.PopSendEvent())
+	{
+		SendNextCustomEvent(*NextToSend);
+	}
 }
 
 void FConcertClientSession::InternalSendCustomRequest(const UScriptStruct* RequestType, const void* RequestData, const FGuid& DestinationEndpointId, const TSharedRef<IConcertSessionCustomResponseHandler>& Handler)
diff --git a/Engine/Plugins/Developer/Concert/ConcertMain/Source/ConcertClient/Private/ConcertClientSession.h b/Engine/Plugins/Developer/Concert/ConcertMain/Source/ConcertClient/Private/ConcertClientSession.h
--- a/Engine/Plugins/Developer/Concert/ConcertMain/Source/ConcertClient/Private/ConcertClientSession.h
+++ b/Engine/Plugins/Developer/Concert/ConcertMain/Source/ConcertClient/Private/ConcertClientSession.h
@@ -3,6 +3,9 @@
 #pragma once
 
 #include "CoreMinimal.h"
+
+#include "Misc/Guid.h"
+#include "IConcertSession.h"
 #include "ConcertSession.h"
 #include "ConcertMessages.h"
 #include "Containers/Ticker.h"
@@ -10,6 +13,7 @@
 
 class IConcertLocalEndpoint;
 struct FConcertClientSettings;
+class FConcertClientSession;
 
 /** Implementation of a Concert Client session */
 class FConcertClientSession : public IConcertClientSession, private FConcertSessionCommonImpl
@@ -55,7 +59,7 @@ public:
 	{
 		return CommonGetScratchpad();
 	}
-	
+
 	virtual FConcertScratchpadPtr GetClientScratchpad(const FGuid& ClientEndpointId) const override
 	{
 		return CommonGetClientScratchpad(ClientEndpointId);
@@ -65,7 +69,7 @@ public:
 	{
 		return ConnectionStatus;
 	}
-	
+
 	virtual FGuid GetSessionClientEndpointId() const override
 	{
 		return ClientSessionEndpoint->GetEndpointContext().EndpointId;
@@ -95,6 +99,11 @@ public:
 	virtual FOnConcertSessionRenamed& OnSessionRenamed() override;
 	virtual FString GetSessionWorkingDirectory() const override;
 
+	virtual FConcertSequencedCustomEventManager& GetSequencedEventManager() override
+	{
+		return CustomEventSequenceManager;
+	}
+
 protected:
 	virtual FDelegateHandle InternalRegisterCustomEventHandler(const FName& EventMessageType, const TSharedRef<IConcertSessionCustomEventHandler>& Handler) override
 	{
@@ -116,7 +125,7 @@ protected:
 		CommonClearCustomEventHandler(EventMessageType);
 	}
 
-	virtual void InternalSendCustomEvent(const UScriptStruct* EventType, const void* EventData, const TArray<FGuid>& DestinationEndpointIds, EConcertMessageFlags Flags) override;
+	virtual void InternalSendCustomEvent(const UScriptStruct* EventType, const void* EventData, const TArray<FGuid>& DestinationEndpointIds, EConcertMessageFlags Flags, TOptional<FConcertSequencedCustomEvent> InSequencedId = {}) override;
 
 	virtual void InternalRegisterCustomRequestHandler(const FName& RequestMessageType, const TSharedRef<IConcertSessionCustomRequestHandler>& Handler) override
 	{
@@ -131,6 +140,9 @@ protected:
 	virtual void InternalSendCustomRequest(const UScriptStruct* RequestType, const void* RequestData, const FGuid& DestinationEndpointId, const TSharedRef<IConcertSessionCustomResponseHandler>& Handler) override;
 
 private:
+	/** */
+	void SendNextCustomEvent(const FConcertSequencedCustomEventManager::FPendingCustomEvent& PendingEvent);
+
 	/** */
 	void HandleRemoteConnectionChanged(const FConcertEndpointContext& RemoteEndpointContext, EConcertRemoteEndpointConnection Connection);
 
@@ -194,6 +206,8 @@ private:
 	/** Callback when the session name changes. */
 	FOnConcertSessionRenamed OnSessionRenamedDelegate;
 
+	FConcertSequencedCustomEventManager CustomEventSequenceManager;
+
 	/** The timespan at which session updates are processed. */
 	const FTimespan SessionTickFrequency;
 
diff --git a/Engine/Plugins/Developer/Concert/ConcertMain/Source/ConcertServer/Private/ConcertServerSession.cpp b/Engine/Plugins/Developer/Concert/ConcertMain/Source/ConcertServer/Private/ConcertServerSession.cpp
--- a/Engine/Plugins/Developer/Concert/ConcertMain/Source/ConcertServer/Private/ConcertServerSession.cpp
+++ b/Engine/Plugins/Developer/Concert/ConcertMain/Source/ConcertServer/Private/ConcertServerSession.cpp
@@ -105,8 +105,10 @@ FString FConcertServerSession::GetSessionWorkingDirectory() const
 	return SessionDirectory;
 }
 
-void FConcertServerSession::InternalSendCustomEvent(const UScriptStruct* EventType, const void* EventData, const TArray<FGuid>& DestinationEndpointIds, EConcertMessageFlags Flags)
+void FConcertServerSession::InternalSendCustomEvent(const UScriptStruct* EventType, const void* EventData, const TArray<FGuid>& DestinationEndpointIds, EConcertMessageFlags Flags, TOptional<FConcertSequencedCustomEvent> InUnused)
 {
+	ensureMsgf(!InUnused, TEXT("Custom event sequencing is not supported on the Multi-user server. Events should arrive in sequence order."));
+
 	if (DestinationEndpointIds.Num() == 0)
 	{
 		return;
diff --git a/Engine/Plugins/Developer/Concert/ConcertMain/Source/ConcertServer/Private/ConcertServerSession.h b/Engine/Plugins/Developer/Concert/ConcertMain/Source/ConcertServer/Private/ConcertServerSession.h
--- a/Engine/Plugins/Developer/Concert/ConcertMain/Source/ConcertServer/Private/ConcertServerSession.h
+++ b/Engine/Plugins/Developer/Concert/ConcertMain/Source/ConcertServer/Private/ConcertServerSession.h
@@ -105,7 +105,7 @@ protected:
 		CommonClearCustomEventHandler(EventMessageType);
 	}
 
-	virtual void InternalSendCustomEvent(const UScriptStruct* EventType, const void* EventData, const TArray<FGuid>& DestinationEndpointIds, EConcertMessageFlags Flags) override;
+	virtual void InternalSendCustomEvent(const UScriptStruct* EventType, const void* EventData, const TArray<FGuid>& DestinationEndpointIds, EConcertMessageFlags Flags, TOptional<FConcertSequencedCustomEvent> InSequencedId = {}) override;
 	
 	virtual void InternalRegisterCustomRequestHandler(const FName& RequestMessageType, const TSharedRef<IConcertSessionCustomRequestHandler>& Handler) override
 	{
diff --git a/Engine/Plugins/Developer/Concert/ConcertSync/ConcertSyncClient/Source/ConcertSyncClient/Private/ConcertClientTransactionManager.cpp b/Engine/Plugins/Developer/Concert/ConcertSync/ConcertSyncClient/Source/ConcertSyncClient/Private/ConcertClientTransactionManager.cpp
--- a/Engine/Plugins/Developer/Concert/ConcertSync/ConcertSyncClient/Source/ConcertSyncClient/Private/ConcertClientTransactionManager.cpp
+++ b/Engine/Plugins/Developer/Concert/ConcertSync/ConcertSyncClient/Source/ConcertSyncClient/Private/ConcertClientTransactionManager.cpp
@@ -136,14 +136,28 @@ void FConcertClientTransactionManager::ReplayTransactions(const FName InPackageN
 		TransactionContext.bIsRequired = true;
 		TransactionContext.PackagesToProcess.Add(InPackageName);
 
+		int64 PackageEventId = 0;
+		LiveSession->GetSessionDatabase().GetPackageHeadEventFromName(InPackageName, PackageEventId);
+
+		FConcertSyncActivity PackageActivity;
+		LiveSession->GetSessionDatabase().GetActivityForEvent(PackageEventId, EConcertSyncActivityEventType::Package, PackageActivity);
+
 		FConcertSyncTransactionEvent TransactionEvent;
 		for (int64 TransactionEventId : LiveTransactionEventIds)
 		{
-			SlowTask.EnterProgressFrame(1.0f, FText::Format(LOCTEXT("ReplayingTransactionForPackageFmt", "Replaying Transaction {0} for {1}"), TransactionEventId, FText::FromName(InPackageName)));
-			
-			if (LiveSession->GetSessionDatabase().GetTransactionEvent(TransactionEventId, TransactionEvent))
+			FConcertSyncActivity TransactionActivity;
+			LiveSession->GetSessionDatabase().GetActivityForEvent(TransactionEventId, EConcertSyncActivityEventType::Transaction, TransactionActivity);
+			if (TransactionActivity.ActivityId > PackageActivity.ActivityId)
 			{
-				PendingTransactionsToProcess.Emplace(TransactionContext, FConcertTransactionFinalizedEvent::StaticStruct(), &TransactionEvent.Transaction);
+				SlowTask.EnterProgressFrame(1.0f, FText::Format(
+												LOCTEXT("ReplayingTransactionForPackageFmt",
+														"Replaying Transaction {0} for {1}"),
+												TransactionEventId, FText::FromName(InPackageName)));
+
+				if (LiveSession->GetSessionDatabase().GetTransactionEvent(TransactionEventId, TransactionEvent))
+				{
+					PendingTransactionsToProcess.Emplace(TransactionContext, FConcertTransactionFinalizedEvent::StaticStruct(), &TransactionEvent.Transaction);
+				}
 			}
 		}
 	}
@@ -198,6 +212,7 @@ void FConcertClientTransactionManager::ProcessPending()
 	{
 		SendPendingTransactionEvents();
 	}
+
 }
 
 template <typename EventType>
@@ -489,7 +504,7 @@ void FConcertClientTransactionManager::FixupObjectIdsInPendingSend(
 		ConflictAggregate.AddObjectRemoved(OldObjectId);
 		for (const FGuid& Guid : *TransactionsWithObject)
 		{
-			PendingTransactionsToSend.Remove(Guid);
+			RemovePendingToSend(Guid);
 		}
 	}
 	else
@@ -558,7 +573,7 @@ void FConcertClientTransactionManager::CheckEventForSendConflicts(const FConcert
 				ConflictAggregate.AddObjectRemoved(ObjectId);
 				for (const FGuid Id : *NamedPrimary)
 				{
-					PendingTransactionsToSend.Remove(Id);
+					RemovePendingToSend(Id);
 				}
 				NameLookupForPendingTransactionsToSend.Remove(ObjectId);
 			}
@@ -637,9 +652,19 @@ FConcertClientTransactionManager::AddPendingToSend(
 	return PendingTransactionsToSend.Emplace(InCommonData.OperationId, InCommonData);
 }
 
+
+void FConcertClientTransactionManager::RemovePendingToSend(const FGuid& InGuid)
+{
+	if (FPendingTransactionToSend* PendingToSend = PendingTransactionsToSend.Find(InGuid))
+	{
+		LiveSession->GetSession().GetSequencedEventManager().RemoveSequencedEvent(PendingToSend->SequencedEvent);
+		PendingTransactionsToSend.Remove(InGuid);
+	}
+}
+
 void FConcertClientTransactionManager::RemovePendingToSend(const FConcertClientLocalTransactionCommonData& InCommonData)
 {
-	PendingTransactionsToSend.Remove(InCommonData.OperationId);
+	RemovePendingToSend(InCommonData.OperationId);
 	if (!CanSendTransactionEvents())
 	{
 		if (UObject* PrimaryObject = InCommonData.PrimaryObject.Get())
@@ -740,9 +765,10 @@ void FConcertClientTransactionManager::HandleLocalTransactionFinalized(const FCo
 	FPendingTransactionToSend& PendingTransaction = HandleLocalTransactionCommon(InCommonData, InFinalizedData.FinalizedObjectUpdates);
 	PendingTransaction.FinalizedData = InFinalizedData;
 	PendingTransaction.bIsFinalized = true;
+	PendingTransaction.SequencedEvent = LiveSession->GetSession().GetSequencedEventManager().AddSequencedCustomEvent();
 }
 
-void FConcertClientTransactionManager::SendTransactionFinalizedEvent(const FGuid& InTransactionId, const FGuid& InOperationId, UObject* InPrimaryObject, const TArray<FName>& InModifiedPackages, const TArray<FConcertExportedObject>& InObjectUpdates, const FConcertLocalIdentifierTable& InLocalIdentifierTable, const FText& InTitle)
+void FConcertClientTransactionManager::SendTransactionFinalizedEvent(const FGuid& InTransactionId, const FGuid& InOperationId, UObject* InPrimaryObject, const TArray<FName>& InModifiedPackages, const TArray<FConcertExportedObject>& InObjectUpdates, const FConcertLocalIdentifierTable& InLocalIdentifierTable, const FConcertSequencedCustomEvent& SequenceId, const FText& InTitle)
 {
 	FConcertTransactionFinalizedEvent TransactionFinalizedEvent;
 	FillTransactionEvent(InTransactionId, InOperationId, InModifiedPackages, TransactionFinalizedEvent);
@@ -751,7 +777,7 @@ void FConcertClientTransactionManager::SendTransactionFinalizedEvent(const FGuid
 	InLocalIdentifierTable.GetState(TransactionFinalizedEvent.LocalIdentifierState);
 	TransactionFinalizedEvent.Title = InTitle;
 
-	LiveSession->GetSession().SendCustomEvent(TransactionFinalizedEvent, LiveSession->GetSession().GetSessionServerEndpointId(), EConcertMessageFlags::ReliableOrdered);
+	LiveSession->GetSession().SendCustomEvent(TransactionFinalizedEvent, LiveSession->GetSession().GetSessionServerEndpointId(), EConcertMessageFlags::ReliableOrdered, SequenceId);
 }
 
 void FConcertClientTransactionManager::SendTransactionSnapshotEvent(const FGuid& InTransactionId, const FGuid& InOperationId, UObject* InPrimaryObject, const TArray<FName>& InModifiedPackages, const TArray<FConcertExportedObject>& InObjectUpdates)
@@ -797,7 +823,8 @@ void FConcertClientTransactionManager::SendPendingTransactionEvents()
 						PrimaryObject, 
 						PendingTransactionPtr->CommonData.ModifiedPackages, 
 						PendingTransactionPtr->FinalizedData.FinalizedObjectUpdates, 
-						PendingTransactionPtr->FinalizedData.FinalizedLocalIdentifierTable, 
+						PendingTransactionPtr->FinalizedData.FinalizedLocalIdentifierTable,
+						PendingTransactionPtr->SequencedEvent,
 						PendingTransactionPtr->CommonData.TransactionTitle
 						);
 				}
@@ -827,7 +854,7 @@ void FConcertClientTransactionManager::SendPendingTransactionEvents()
 		{
 			// TODO: Broadcast delegate
 
-			PendingTransactionsToSend.Remove(PendingTransactionPtr->CommonData.TransactionId);
+			RemovePendingToSend(PendingTransactionPtr->CommonData.TransactionId);
 			PendingTransactionsToSendOrderIter.RemoveCurrent();
 			continue;
 		}
diff --git a/Engine/Plugins/Developer/Concert/ConcertSync/ConcertSyncClient/Source/ConcertSyncClient/Private/ConcertClientTransactionManager.h b/Engine/Plugins/Developer/Concert/ConcertSync/ConcertSyncClient/Source/ConcertSyncClient/Private/ConcertClientTransactionManager.h
--- a/Engine/Plugins/Developer/Concert/ConcertSync/ConcertSyncClient/Source/ConcertSyncClient/Private/ConcertClientTransactionManager.h
+++ b/Engine/Plugins/Developer/Concert/ConcertSync/ConcertSyncClient/Source/ConcertSyncClient/Private/ConcertClientTransactionManager.h
@@ -7,6 +7,7 @@
 #include "Misc/ITransaction.h"
 #include "UObject/Class.h"
 #include "ConcertMessages.h"
+#include "IConcertSession.h"
 #include "IConcertSessionHandler.h"
 #include "IConcertClientTransactionBridge.h"
 #include "UObject/StructOnScope.h"
@@ -136,6 +137,10 @@ private:
 		FConcertClientLocalTransactionCommonData CommonData;
 		FConcertClientLocalTransactionSnapshotData SnapshotData;
 		FConcertClientLocalTransactionFinalizedData FinalizedData;
+
+		// Sequencing id for this pending transaction.
+		FConcertSequencedCustomEvent SequencedEvent;
+
 		bool bIsFinalized = false;
 		double LastSnapshotTimeSeconds = 0.0;
 	};
@@ -195,7 +200,7 @@ private:
 	/**
 	 * Send a transaction finalized event.
 	 */
-	void SendTransactionFinalizedEvent(const FGuid& InTransactionId, const FGuid& InOperationId, UObject* InPrimaryObject, const TArray<FName>& InModifiedPackages, const TArray<FConcertExportedObject>& InObjectUpdates, const FConcertLocalIdentifierTable& InLocalIdentifierTable, const FText& InTitle);
+	void SendTransactionFinalizedEvent(const FGuid& InTransactionId, const FGuid& InOperationId, UObject* InPrimaryObject, const TArray<FName>& InModifiedPackages, const TArray<FConcertExportedObject>& InObjectUpdates, const FConcertLocalIdentifierTable& InLocalIdentifierTable, const FConcertSequencedCustomEvent& SequenceId, const FText& InTitle);
 
 	/**
 	 * Send a transaction snapshot event.
@@ -228,6 +233,11 @@ private:
 	 */
 	void RemovePendingToSend(const FConcertClientLocalTransactionCommonData& InCommonData);
 
+	/**
+	 * Remove the given pending to send object by Guid.
+	 */
+	void RemovePendingToSend(const FGuid& InCommonData);
+
 	/**
 	 * Remove the given pending to send object.
 	 */
diff --git a/Engine/Plugins/Developer/Concert/ConcertSync/ConcertSyncClient/Source/ConcertSyncClient/Test/ConcertClientLiveTransactionAuthorsTests.cpp b/Engine/Plugins/Developer/Concert/ConcertSync/ConcertSyncClient/Source/ConcertSyncClient/Test/ConcertClientLiveTransactionAuthorsTests.cpp
--- a/Engine/Plugins/Developer/Concert/ConcertSync/ConcertSyncClient/Source/ConcertSyncClient/Test/ConcertClientLiveTransactionAuthorsTests.cpp
+++ b/Engine/Plugins/Developer/Concert/ConcertSync/ConcertSyncClient/Source/ConcertSyncClient/Test/ConcertClientLiveTransactionAuthorsTests.cpp
@@ -47,7 +47,10 @@ public:
 	virtual void InternalUnregisterCustomEventHandler(const FName& EventMessageType, const FDelegateHandle EventHandle) override			{ return NotMocked<void>(); }
 	virtual void InternalUnregisterCustomEventHandler(const FName& EventMessageType, const void* EventHandler) override						{ return NotMocked<void>(); }
 	virtual void InternalClearCustomEventHandler(const FName& EventMessageType) override													{ return NotMocked<void>(); }
-	virtual void InternalSendCustomEvent(const UScriptStruct*, const void*, const TArray<FGuid>&, EConcertMessageFlags) override			{ return NotMocked<void>(); }
+	virtual void InternalSendCustomEvent(const UScriptStruct*, const void*, const TArray<FGuid>&, EConcertMessageFlags, TOptional<FConcertSequencedCustomEvent> = {}) override
+	{
+		return NotMocked<void>();
+	}
 	virtual void InternalRegisterCustomRequestHandler(const FName&, const TSharedRef<IConcertSessionCustomRequestHandler>&) override		{ return NotMocked<void>(); }
 	virtual void InternalUnregisterCustomRequestHandler(const FName&) override																{ return NotMocked<void>(); }
 	virtual void InternalSendCustomRequest(const UScriptStruct*, const void*, const FGuid&, const TSharedRef<IConcertSessionCustomResponseHandler>&) override { NotMocked<void>(); }
@@ -89,6 +92,10 @@ public:
 		}
 		return false;
 	}
+	virtual FConcertSequencedCustomEventManager& GetSequencedEventManager() override
+	{
+		return NotMocked<FConcertSequencedCustomEventManager&>(CustomEventSequenceManager);
+	}
 
 protected:
 	FGuid Id;
@@ -101,6 +108,7 @@ protected:
 	FOnConcertClientSessionClientChanged ClientChanged;
 	FOnConcertSessionRenamed SessionRenamed;
 	FConcertSyncEndpointIdAndData LocalClientInfo;
+	FConcertSequencedCustomEventManager CustomEventSequenceManager;
 	TMap<FGuid, FConcertSessionClientInfo> OtherClientsInfo;
 };
 
diff --git a/Engine/Plugins/Developer/Concert/ConcertSync/ConcertSyncCore/Source/ConcertSyncCore/Private/ConcertSyncSessionDatabase.cpp b/Engine/Plugins/Developer/Concert/ConcertSync/ConcertSyncCore/Source/ConcertSyncCore/Private/ConcertSyncSessionDatabase.cpp
--- a/Engine/Plugins/Developer/Concert/ConcertSync/ConcertSyncCore/Source/ConcertSyncCore/Private/ConcertSyncSessionDatabase.cpp
+++ b/Engine/Plugins/Developer/Concert/ConcertSync/ConcertSyncCore/Source/ConcertSyncCore/Private/ConcertSyncSessionDatabase.cpp
@@ -2486,6 +2486,17 @@ bool FConcertSyncSessionDatabase::GetPackageEventMetaData(const int64 InPackageE
 	return Statements->GetPackageEventForId(InPackageEventId, OuptPackageRevision, OutPackageInfo, UnusedDataFilename);
 }
 
+bool FConcertSyncSessionDatabase::GetPackageHeadEventFromName(const FName InPackageName, int64& OutPackageEventId)
+{
+	int64 PackageNameId = 0;
+	if (!GetPackageNameId(InPackageName, PackageNameId))
+	{
+		return false;
+	}
+
+	return Statements->GetPackageHeadEventId(PackageNameId, OutPackageEventId);
+}
+
 bool FConcertSyncSessionDatabase::GetPackageEvent(const int64 InPackageEventId, const TFunctionRef<void(FConcertSyncPackageEventData&)>& PackageEventFn) const
 {
 	FString DataFilename;
diff --git a/Engine/Plugins/Developer/Concert/ConcertSync/ConcertSyncCore/Source/ConcertSyncCore/Public/ConcertSyncSessionDatabase.h b/Engine/Plugins/Developer/Concert/ConcertSync/ConcertSyncCore/Source/ConcertSyncCore/Public/ConcertSyncSessionDatabase.h
--- a/Engine/Plugins/Developer/Concert/ConcertSync/ConcertSyncCore/Source/ConcertSyncCore/Public/ConcertSyncSessionDatabase.h
+++ b/Engine/Plugins/Developer/Concert/ConcertSync/ConcertSyncCore/Source/ConcertSyncCore/Public/ConcertSyncSessionDatabase.h
@@ -585,7 +585,17 @@ public:
 	 * @return True if the package event was found, false otherwise.
 	 */
 	bool GetPackageEventMetaData(const int64 InPackageEventId, int64& OuptPackageRevision, FConcertPackageInfo& OutPackageInfo) const;
-	
+
+	/**
+	 * Get Package Event Id from Package Name.  The event id will be from the head revision of the package.
+	 *
+	 * @param InPackageName
+	 * @param OutPackageEventId
+	 *
+	 * @return True if the package event was found, false otherwise
+	 */
+	bool GetPackageHeadEventFromName(const FName InPackageName, int64& OutPackageEventId);
+
 	/**
 	 * Get a package event from this database.
 	 *
diff --git a/Engine/Plugins/Developer/Concert/ConcertSync/ConcertSyncTest/Source/ConcertSyncTest/Private/ConcertDataStoreTests.cpp b/Engine/Plugins/Developer/Concert/ConcertSync/ConcertSyncTest/Source/ConcertSyncTest/Private/ConcertDataStoreTests.cpp
--- a/Engine/Plugins/Developer/Concert/ConcertSync/ConcertSyncTest/Source/ConcertSyncTest/Private/ConcertDataStoreTests.cpp
+++ b/Engine/Plugins/Developer/Concert/ConcertSync/ConcertSyncTest/Source/ConcertSyncTest/Private/ConcertDataStoreTests.cpp
@@ -93,7 +93,7 @@ public:
 	virtual void InternalUnregisterCustomEventHandler(const FName& EventMessageType, const FDelegateHandle EventHandle) override			{ return NotMocked<void>(); }
 	virtual void InternalUnregisterCustomEventHandler(const FName& EventMessageType, const void* EventHandler) override						{ return NotMocked<void>(); }
 	virtual void InternalClearCustomEventHandler(const FName& EventMessageType) override													{ return NotMocked<void>(); }
-	virtual void InternalSendCustomEvent(const UScriptStruct*, const void*, const TArray<FGuid>&, EConcertMessageFlags) override			{ return NotMocked<void>(); }
+	virtual void InternalSendCustomEvent(const UScriptStruct*, const void*, const TArray<FGuid>&, EConcertMessageFlags, TOptional<FConcertSequencedCustomEvent> InSequenceId={}) override			{ return NotMocked<void>(); }
 	virtual void InternalRegisterCustomRequestHandler(const FName&, const TSharedRef<IConcertSessionCustomRequestHandler>&) override		{ return NotMocked<void>(); }
 	virtual void InternalUnregisterCustomRequestHandler(const FName&) override																{ return NotMocked<void>(); }
 	virtual void InternalSendCustomRequest(const UScriptStruct*, const void*, const FGuid&, const TSharedRef<IConcertSessionCustomResponseHandler>&) override { NotMocked<void>(); }
@@ -141,12 +141,16 @@ public:
 	virtual void InternalUnregisterCustomEventHandler(const FName& EventMessageType, const FDelegateHandle EventHandle) override			{ return NotMocked<void>(); }
 	virtual void InternalUnregisterCustomEventHandler(const FName& EventMessageType, const void* EventHandler) override						{ return NotMocked<void>(); }
 	virtual void InternalClearCustomEventHandler(const FName& EventMessageType) override													{ return NotMocked<void>(); }
-	virtual void InternalSendCustomEvent(const UScriptStruct*, const void*, const TArray<FGuid>&, EConcertMessageFlags) override			{ return NotMocked<void>(); }
+	virtual void InternalSendCustomEvent(const UScriptStruct*, const void*, const TArray<FGuid>&, EConcertMessageFlags, TOptional<FConcertSequencedCustomEvent> InSequenceId={}) override			{ return NotMocked<void>(); }
 	virtual void InternalRegisterCustomRequestHandler(const FName&, const TSharedRef<IConcertSessionCustomRequestHandler>&) override		{ return NotMocked<void>(); }
 	virtual void InternalUnregisterCustomRequestHandler(const FName&) override																{ return NotMocked<void>(); }
 	virtual void InternalSendCustomRequest(const UScriptStruct*, const void*, const FGuid&, const TSharedRef<IConcertSessionCustomResponseHandler>&) override { NotMocked<void>(); }
 	// IConcertSession End.
 
+	virtual FConcertSequencedCustomEventManager& GetSequencedEventManager() override
+	{
+		return NotMocked<FConcertSequencedCustomEventManager&>(CustomEventSequenceManager);
+	}
 	// IConcertClientSession Begin
 	virtual EConcertConnectionStatus GetConnectionStatus() const override            { return NotMocked<EConcertConnectionStatus>(EConcertConnectionStatus::Connected); }
 	virtual FGuid GetSessionClientEndpointId() const override                        { return NotMocked<FGuid>(); }
@@ -184,6 +188,7 @@ protected:
 	FOnConcertClientSessionConnectionChanged ConnectionChanged;
 	FOnConcertClientSessionClientChanged ClientChanged;
 	FOnConcertSessionRenamed SessionRenamed;
+	FConcertSequencedCustomEventManager CustomEventSequenceManager;
 	FConcertClientInfo ClientInfo;
 };
 
@@ -191,7 +196,7 @@ protected:
 class FConcertServerSessionMock : public FConcertServerSessionBaseMock
 {
 public:
-	virtual void InternalSendCustomEvent(const UScriptStruct* EventType, const void* EventData, const TArray<FGuid>& TargetEndpointIds, EConcertMessageFlags)
+	virtual void InternalSendCustomEvent(const UScriptStruct* EventType, const void* EventData, const TArray<FGuid>& TargetEndpointIds, EConcertMessageFlags, TOptional<FConcertSequencedCustomEvent> InSequenceId={}) override
 	{
 		for (const FGuid& TargetEndPointId : TargetEndpointIds)
 		{
