TObjectPtr: Catching up on new Engine/Plugins AddReferencedObjects calls

diff --git a/Engine/Plugins/Developer/Concert/ConcertSync/ConcertSyncClient/Source/ConcertSyncClient/Private/ConcertClientSequencerManager.cpp b/Engine/Plugins/Developer/Concert/ConcertSync/ConcertSyncClient/Source/ConcertSyncClient/Private/ConcertClientSequencerManager.cpp
--- a/Engine/Plugins/Developer/Concert/ConcertSync/ConcertSyncClient/Source/ConcertSyncClient/Private/ConcertClientSequencerManager.cpp
+++ b/Engine/Plugins/Developer/Concert/ConcertSync/ConcertSyncClient/Source/ConcertSyncClient/Private/ConcertClientSequencerManager.cpp
@@ -535,7 +535,7 @@ void FConcertClientSequencerManager::HandleAssetReload(const EPackageReloadPhase
 		return;
 	}
 
-	for (TTuple<FName, ALevelSequenceActor*>& Item : SequencePlayers)
+	for (TTuple<FName, TObjectPtr<ALevelSequenceActor>>& Item : SequencePlayers)
 	{
 		ALevelSequenceActor* LevelSequenceActor = Item.Value;
 		// If we have a null LevelSequenceActor it means that it has already been destroyed by a close event.  We will not
@@ -559,7 +559,7 @@ void FConcertClientSequencerManager::ApplyCloseEventToPlayers(const FConcertSequ
 {
 	UE_LOG(LogConcertSequencerSync, Verbose, TEXT("    Closing sequence players: %s, is from controller: %d"), *CloseEvent.SequenceObjectPath, CloseEvent.bControllerClose);
 
-	ALevelSequenceActor** Player = SequencePlayers.Find(*CloseEvent.SequenceObjectPath);
+	TObjectPtr<ALevelSequenceActor>* Player = SequencePlayers.Find(*CloseEvent.SequenceObjectPath);
 	if (!Player)
 	{
 		UE_LOG(LogConcertSequencerSync, Verbose, TEXT("        No open players to close for sequence %s"), *CloseEvent.SequenceObjectPath);
@@ -796,7 +796,7 @@ void FConcertClientSequencerManager::ApplyTimeAdjustmentToPlayers(const FConcert
 {
 	UE_LOG(LogConcertSequencerSync, Verbose, TEXT("    TimeAdjustment: Updating sequence players for sequence %s"), *TimeAdjustmentEvent.SequenceObjectPath);
 
-	ALevelSequenceActor** SeqPlayer = SequencePlayers.Find(*TimeAdjustmentEvent.SequenceObjectPath);
+	TObjectPtr<ALevelSequenceActor>* SeqPlayer = SequencePlayers.Find(*TimeAdjustmentEvent.SequenceObjectPath);
 	if (!SeqPlayer)
 	{
 		UE_LOG(LogConcertSequencerSync, Verbose, TEXT("        No sequence player for sequence %s"), *TimeAdjustmentEvent.SequenceObjectPath);
@@ -824,7 +824,7 @@ void FConcertClientSequencerManager::OnPrecacheEvent(const FConcertSessionContex
 			*SequenceObjectPathStr, bShouldPrecache ? TEXT("added to") : TEXT("removed from"));
 		if (bShouldPrecache)
 		{
-			if (ULevelSequence** ExistingSequence = PrecachedSequences.Find(SequenceObjectPath);
+			if (TObjectPtr<ULevelSequence>* ExistingSequence = PrecachedSequences.Find(SequenceObjectPath);
 				ExistingSequence && *ExistingSequence)
 			{
 				UE_LOG(LogConcertSequencerSync, Warning, TEXT("OnPrecacheEvent: %s already in precache set"), *SequenceObjectPath.ToString());
@@ -974,7 +974,7 @@ void FConcertClientSequencerManager::ApplyEventToPlayers(const FConcertSequencer
 {
 	UE_LOG(LogConcertSequencerSync, Verbose, TEXT("    Transport: Update sequence player for sequence %s, at frame: %d"), *EventState.SequenceObjectPath, EventState.Time.Time.FrameNumber.Value);
 
-	ALevelSequenceActor** SeqPlayer = SequencePlayers.Find(*EventState.SequenceObjectPath);
+	TObjectPtr<ALevelSequenceActor>* SeqPlayer = SequencePlayers.Find(*EventState.SequenceObjectPath);
 	if (!SeqPlayer)
 	{
 		UE_LOG(LogConcertSequencerSync, Verbose, TEXT("        No sequence player for sequence %s"), *EventState.SequenceObjectPath);
@@ -1092,7 +1092,7 @@ void FConcertClientSequencerManager::OnWorkspaceEndFrameCompleted()
 	{
 		for (const TTuple<FName, FString>& Player : PendingDestroy)
 		{
-			if (ALevelSequenceActor** SequencePlayer = SequencePlayers.Find(Player.Get<0>()))
+			if (TObjectPtr<ALevelSequenceActor>* SequencePlayer = SequencePlayers.Find(Player.Get<0>()))
 			{
 				DestroyPlayer(*SequencePlayer);
 				SequencePlayers.Remove(Player.Get<0>());
diff --git a/Engine/Plugins/Developer/Concert/ConcertSync/ConcertSyncClient/Source/ConcertSyncClient/Private/ConcertClientSequencerManager.h b/Engine/Plugins/Developer/Concert/ConcertSync/ConcertSyncClient/Source/ConcertSyncClient/Private/ConcertClientSequencerManager.h
--- a/Engine/Plugins/Developer/Concert/ConcertSync/ConcertSyncClient/Source/ConcertSyncClient/Private/ConcertClientSequencerManager.h
+++ b/Engine/Plugins/Developer/Concert/ConcertSync/ConcertSyncClient/Source/ConcertSyncClient/Private/ConcertClientSequencerManager.h
@@ -360,10 +360,10 @@ private:
 	TArray<FOpenSequencerData> OpenSequencers;
 
 	/** Map of opened sequence players, if not in editor mode. */
-	TMap<FName, ALevelSequenceActor*> SequencePlayers;
+	TMap<FName, TObjectPtr<ALevelSequenceActor>> SequencePlayers;
 
 	/** Map of precached sequence assets. */
-	TMap<FName, ULevelSequence*> PrecachedSequences;
+	TMap<FName, TObjectPtr<ULevelSequence>> PrecachedSequences;
 
 	/** Boolean that is set when we are handling any transport event to prevent re-entrancy */
 	bool bRespondingToTransportEvent;
diff --git a/Engine/Plugins/Experimental/ChaosVD/Source/ChaosVD/Private/ChaosVDGeometryBuilder.cpp b/Engine/Plugins/Experimental/ChaosVD/Source/ChaosVD/Private/ChaosVDGeometryBuilder.cpp
--- a/Engine/Plugins/Experimental/ChaosVD/Source/ChaosVD/Private/ChaosVDGeometryBuilder.cpp
+++ b/Engine/Plugins/Experimental/ChaosVD/Source/ChaosVD/Private/ChaosVDGeometryBuilder.cpp
@@ -185,14 +185,14 @@ void FChaosVDGeometryBuilder::ApplyMeshToComponentFromKey(TWeakObjectPtr<UMeshCo
 
 	if (UDynamicMeshComponent* DynamicMeshComponent = Cast<UDynamicMeshComponent>(MeshComponent))
 	{
-		if (UDynamicMesh** DynamicMesh = DynamicMeshCacheMap.Find(GeometryKey))
+		if (TObjectPtr<UDynamicMesh>* DynamicMesh = DynamicMeshCacheMap.Find(GeometryKey))
 		{
 			DynamicMeshComponent->SetDynamicMesh(*DynamicMesh);
 		}
 	}
 	else if (UStaticMeshComponent* StaticMeshComponent = Cast<UStaticMeshComponent>(MeshComponent))
 	{
-		if (UStaticMesh** StaticMesh = StaticMeshCacheMap.Find(GeometryKey))
+		if (TObjectPtr<UStaticMesh>* StaticMesh = StaticMeshCacheMap.Find(GeometryKey))
 		{
 			StaticMeshComponent->SetStaticMesh(*StaticMesh);
 		}
diff --git a/Engine/Plugins/Experimental/ChaosVD/Source/ChaosVD/Public/ChaosVDGeometryBuilder.h b/Engine/Plugins/Experimental/ChaosVD/Source/ChaosVD/Public/ChaosVDGeometryBuilder.h
--- a/Engine/Plugins/Experimental/ChaosVD/Source/ChaosVD/Public/ChaosVDGeometryBuilder.h
+++ b/Engine/Plugins/Experimental/ChaosVD/Source/ChaosVD/Public/ChaosVDGeometryBuilder.h
@@ -163,10 +163,10 @@ private:
 	void RegisterMeshComponentWaitingForGeometry(uint32 GeometryKey, TWeakObjectPtr<UMeshComponent> MesComponent, const int32 LODsToGenerateNum);
 
 	/** Map containing already generated dynamic mesh for any given implicit object */
-	TMap<uint32, UDynamicMesh*> DynamicMeshCacheMap;
+	TMap<uint32, TObjectPtr<UDynamicMesh>> DynamicMeshCacheMap;
 
 	/** Map containing already generated static mesh for any given implicit object */
-	TMap<uint32, UStaticMesh*> StaticMeshCacheMap;
+	TMap<uint32, TObjectPtr<UStaticMesh>> StaticMeshCacheMap;
 
 	/** Map containing all the meshes component waiting for geometry, by geometry key*/
 	TMap<uint32, TArray<TWeakObjectPtr<UMeshComponent>>> MeshComponentsWaitingForGeometryByKey;
@@ -457,14 +457,14 @@ MeshType* FChaosVDGeometryBuilder::GetCachedMeshForImplicit(const uint32 Geometr
 {
 	if constexpr (std::is_same_v<MeshType, UDynamicMesh>)
 	{
-		if (MeshType** MeshPtrPtr = DynamicMeshCacheMap.Find(GeometryCacheKey))
+		if (TObjectPtr<MeshType>* MeshPtrPtr = DynamicMeshCacheMap.Find(GeometryCacheKey))
 		{
 			return *MeshPtrPtr;
 		}
 	}
 	else if constexpr (std::is_same_v<MeshType, UStaticMesh>)
 	{
-		if (MeshType** MeshPtrPtr = StaticMeshCacheMap.Find(GeometryCacheKey))
+		if (TObjectPtr<MeshType>* MeshPtrPtr = StaticMeshCacheMap.Find(GeometryCacheKey))
 		{
 			return *MeshPtrPtr;
 		}
diff --git a/Engine/Plugins/Experimental/ChaosVD/Source/ChaosVD/Public/ChaosVDScene.h b/Engine/Plugins/Experimental/ChaosVD/Source/ChaosVD/Public/ChaosVDScene.h
--- a/Engine/Plugins/Experimental/ChaosVD/Source/ChaosVD/Public/ChaosVDScene.h
+++ b/Engine/Plugins/Experimental/ChaosVD/Source/ChaosVD/Public/ChaosVDScene.h
@@ -100,7 +100,7 @@ private:
 	void HandlePostSelectionChange(const UTypedElementSelectionSet* PreChangeSelectionSet);
 
 	/** UWorld instance used to represent the recorded debug data */
-	UWorld* PhysicsVDWorld = nullptr;
+	TObjectPtr<UWorld> PhysicsVDWorld = nullptr;
 
 	FChaosVDSceneUpdatedDelegate SceneUpdatedDelegate;
 
@@ -111,7 +111,7 @@ private:
 	FChaosVDOnObjectSelectedDelegate ObjectFocusedDelegate;
 
 	/** Selection set object holding the current selection state */
-	UTypedElementSelectionSet* SelectionSet;
+	TObjectPtr<UTypedElementSelectionSet> SelectionSet;
 
 	/** Array of actors with hit proxies that need to be updated */
 	TArray<AActor*> PendingActorsToUpdateSelectionProxy;
diff --git a/Engine/Plugins/Experimental/Mutable/Source/CustomizableObject/Private/MuCO/UnrealMutableImageProvider.h b/Engine/Plugins/Experimental/Mutable/Source/CustomizableObject/Private/MuCO/UnrealMutableImageProvider.h
--- a/Engine/Plugins/Experimental/Mutable/Source/CustomizableObject/Private/MuCO/UnrealMutableImageProvider.h
+++ b/Engine/Plugins/Experimental/Mutable/Source/CustomizableObject/Private/MuCO/UnrealMutableImageProvider.h
@@ -69,7 +69,7 @@ private:
 		/** If the above Image has not been loaded in the game thread, the TextureToLoad bulk data will be loaded
 		* from the Mutable thread when it's needed
 		*/
-		UTexture2D* TextureToLoad = nullptr;
+		TObjectPtr<UTexture2D> TextureToLoad = nullptr;
 
 		/** true of the reference maintained by the user. */
 		bool ReferencesUser = false;
diff --git a/Engine/Plugins/Experimental/Mutable/Source/CustomizableObjectEditor/Private/MuCOE/SMutableMeshViewport.h b/Engine/Plugins/Experimental/Mutable/Source/CustomizableObjectEditor/Private/MuCOE/SMutableMeshViewport.h
--- a/Engine/Plugins/Experimental/Mutable/Source/CustomizableObjectEditor/Private/MuCOE/SMutableMeshViewport.h
+++ b/Engine/Plugins/Experimental/Mutable/Source/CustomizableObjectEditor/Private/MuCOE/SMutableMeshViewport.h
@@ -78,7 +78,7 @@ private:
 	 */
 
 	/** Mutable mesh converted into a Unreal mesh object */
-	USkeletalMeshComponent* SkeletalMeshComponent = nullptr;
+	TObjectPtr<USkeletalMeshComponent> SkeletalMeshComponent = nullptr;
 	
 	/** Generates a new USkeletalMesh from the MutableMesh set on this object
 	 * @return - If the operation was able to be processed or not
diff --git a/Engine/Plugins/FX/Niagara/Source/NiagaraEditor/Private/ViewModels/NiagaraSimCacheViewModel.h b/Engine/Plugins/FX/Niagara/Source/NiagaraEditor/Private/ViewModels/NiagaraSimCacheViewModel.h
--- a/Engine/Plugins/FX/Niagara/Source/NiagaraEditor/Private/ViewModels/NiagaraSimCacheViewModel.h
+++ b/Engine/Plugins/FX/Niagara/Source/NiagaraEditor/Private/ViewModels/NiagaraSimCacheViewModel.h
@@ -118,10 +118,10 @@ private:
 	void BuildComponentInfos(const FName Name, const UScriptStruct* Struct, TArray<FComponentInfo>& ComponentInfos);
 	
 	// The sim cache being viewed
-	UNiagaraSimCache* SimCache = nullptr;
+	TObjectPtr<UNiagaraSimCache> SimCache = nullptr;
 
 	// Component for preview scene
-	UNiagaraComponent* PreviewComponent = nullptr;
+	TObjectPtr<UNiagaraComponent> PreviewComponent = nullptr;
 
 	// Which frame of the cached sim is being viewed
 	int32 FrameIndex = 0;
diff --git a/Engine/Plugins/Runtime/CommonUI/Source/CommonUI/Public/WidgetFactory.h b/Engine/Plugins/Runtime/CommonUI/Source/CommonUI/Public/WidgetFactory.h
--- a/Engine/Plugins/Runtime/CommonUI/Source/CommonUI/Public/WidgetFactory.h
+++ b/Engine/Plugins/Runtime/CommonUI/Source/CommonUI/Public/WidgetFactory.h
@@ -272,8 +272,8 @@ public:
 	}
 
 private:
-	TArray<WidgetType*>	ActiveWidgets;
-	TArray<WidgetType*>	InactiveWidgets;
+	TArray<TObjectPtr<WidgetType>>	ActiveWidgets;
+	TArray<TObjectPtr<WidgetType>>	InactiveWidgets;
 	TMap<WidgetType*, TSharedPtr<SWidget>>	CachedSlateWidgets;
 	TWeakObjectPtr<UWidget> OwningWidget;
 	TSubclassOf<WidgetType>	WidgetClass;
diff --git a/Engine/Plugins/VirtualProduction/DMX/DMXPixelMapping/Source/DMXPixelMappingEditor/Private/Views/SDMXPixelMappingDMXLibraryView.cpp b/Engine/Plugins/VirtualProduction/DMX/DMXPixelMapping/Source/DMXPixelMappingEditor/Private/Views/SDMXPixelMappingDMXLibraryView.cpp
--- a/Engine/Plugins/VirtualProduction/DMX/DMXPixelMapping/Source/DMXPixelMappingEditor/Private/Views/SDMXPixelMappingDMXLibraryView.cpp
+++ b/Engine/Plugins/VirtualProduction/DMX/DMXPixelMapping/Source/DMXPixelMappingEditor/Private/Views/SDMXPixelMappingDMXLibraryView.cpp
@@ -213,7 +213,7 @@ void SDMXPixelMappingDMXLibraryView::OnComponentSelected()
 		UDMXPixelMappingBaseComponent* Component = ComponentReference.GetComponent();
 		do
 		{
-			UDMXPixelMappingDMXLibraryViewModel* const* CorrespondingViewModelPtr = Algo::FindByPredicate(ViewModels, [Component](const UDMXPixelMappingDMXLibraryViewModel* ViewModel)
+			TObjectPtr<UDMXPixelMappingDMXLibraryViewModel>* CorrespondingViewModelPtr = Algo::FindByPredicate(ViewModels, [Component](const UDMXPixelMappingDMXLibraryViewModel* ViewModel)
 				{
 					return ViewModel->GetFixtureGroupComponent() == Component;
 				});
diff --git a/Engine/Plugins/VirtualProduction/DMX/DMXPixelMapping/Source/DMXPixelMappingEditor/Private/Views/SDMXPixelMappingDMXLibraryView.h b/Engine/Plugins/VirtualProduction/DMX/DMXPixelMapping/Source/DMXPixelMappingEditor/Private/Views/SDMXPixelMappingDMXLibraryView.h
--- a/Engine/Plugins/VirtualProduction/DMX/DMXPixelMapping/Source/DMXPixelMappingEditor/Private/Views/SDMXPixelMappingDMXLibraryView.h
+++ b/Engine/Plugins/VirtualProduction/DMX/DMXPixelMapping/Source/DMXPixelMappingEditor/Private/Views/SDMXPixelMappingDMXLibraryView.h
@@ -78,7 +78,7 @@ private:
 	FTimerHandle RefreshTimerHandle;
 
 	/** The view model of the focused DMX Library */
-	TArray<UDMXPixelMappingDMXLibraryViewModel*> ViewModels;
+	TArray<TObjectPtr<UDMXPixelMappingDMXLibraryViewModel>> ViewModels;
 
 	/** The toolkit of this editor */
 	TWeakPtr<FDMXPixelMappingToolkit> WeakToolkit;
diff --git a/Engine/Source/Editor/Kismet/Public/SKismetInspector.h b/Engine/Source/Editor/Kismet/Public/SKismetInspector.h
--- a/Engine/Source/Editor/Kismet/Public/SKismetInspector.h
+++ b/Engine/Source/Editor/Kismet/Public/SKismetInspector.h
@@ -181,7 +181,7 @@ protected:
 	bool bRefreshOnTick;
 
 	/** Holds the property objects that need to be displayed by the inspector starting on the next tick */
-	TArray<UObject*> RefreshPropertyObjects;
+	TArray<TObjectPtr<UObject>> RefreshPropertyObjects;
 
 	/** Details options that are used by the inspector on the next refresh. */
 	FShowDetailsOptions RefreshOptions;
diff --git a/Engine/Source/Editor/Persona/Private/SAnimAssetFindReplace.cpp b/Engine/Source/Editor/Persona/Private/SAnimAssetFindReplace.cpp
--- a/Engine/Source/Editor/Persona/Private/SAnimAssetFindReplace.cpp
+++ b/Engine/Source/Editor/Persona/Private/SAnimAssetFindReplace.cpp
@@ -1,4 +1,4 @@
-﻿// Copyright Epic Games, Inc. All Rights Reserved.
+// Copyright Epic Games, Inc. All Rights Reserved.
 
 #include "SAnimAssetFindReplace.h"
 
@@ -414,7 +414,7 @@ void SAnimAssetFindReplace::Construct(const FArguments& InArgs)
 
 void SAnimAssetFindReplace::SetCurrentProcessor(TSubclassOf<UAnimAssetFindReplaceProcessor> InProcessorClass)
 {
-	UAnimAssetFindReplaceProcessor** FoundProcessor = Processors.Find(InProcessorClass);
+	TObjectPtr<UAnimAssetFindReplaceProcessor>* FoundProcessor = Processors.Find(InProcessorClass);
 	if(FoundProcessor)
 	{
 		CurrentProcessor = *FoundProcessor;
@@ -425,7 +425,7 @@ void SAnimAssetFindReplace::SetCurrentProcessor(TSubclassOf<UAnimAssetFindReplac
 
 UAnimAssetFindReplaceProcessor* SAnimAssetFindReplace::GetProcessor(TSubclassOf<UAnimAssetFindReplaceProcessor> InProcessorClass) const
 {
-	UAnimAssetFindReplaceProcessor* const* FoundProcessor = Processors.Find(InProcessorClass);
+	const TObjectPtr<UAnimAssetFindReplaceProcessor>* FoundProcessor = Processors.Find(InProcessorClass);
 	return FoundProcessor ? *FoundProcessor : nullptr;
 }
 
@@ -601,4 +601,4 @@ void SAnimAssetFindReplace::RemoveInAsset(const FAssetData& InAssetData) const
 	CurrentProcessor->RemoveInAsset(InAssetData);
 }
 
-#undef LOCTEXT_NAMESPACE
\ No newline at end of file
+#undef LOCTEXT_NAMESPACE
diff --git a/Engine/Source/Editor/Persona/Private/SAnimAssetFindReplace.h b/Engine/Source/Editor/Persona/Private/SAnimAssetFindReplace.h
--- a/Engine/Source/Editor/Persona/Private/SAnimAssetFindReplace.h
+++ b/Engine/Source/Editor/Persona/Private/SAnimAssetFindReplace.h
@@ -89,7 +89,7 @@ private:
 
 	EAnimAssetFindReplaceMode Mode = EAnimAssetFindReplaceMode::Find;
 
-	TMap<TSubclassOf<UAnimAssetFindReplaceProcessor>, UAnimAssetFindReplaceProcessor*> Processors;
+	TMap<TSubclassOf<UAnimAssetFindReplaceProcessor>, TObjectPtr<UAnimAssetFindReplaceProcessor>> Processors;
 
 	UAnimAssetFindReplaceProcessor* CurrentProcessor = nullptr;
 
@@ -109,7 +109,7 @@ private:
 
 	TArray<FAssetData> OldAssets;
 
-	UAnimAssetFindReplaceContext* ToolbarContext;
+	TObjectPtr<UAnimAssetFindReplaceContext> ToolbarContext;
 
 	TSharedPtr<SBox> ToolbarContainer;
 
diff --git a/Engine/Source/Runtime/Engine/Private/WorldPartition/WorldPartition.cpp b/Engine/Source/Runtime/Engine/Private/WorldPartition/WorldPartition.cpp
--- a/Engine/Source/Runtime/Engine/Private/WorldPartition/WorldPartition.cpp
+++ b/Engine/Source/Runtime/Engine/Private/WorldPartition/WorldPartition.cpp
@@ -1485,7 +1485,7 @@ void UWorldPartition::Tick(float DeltaSeconds)
 		EditorHash->Tick(DeltaSeconds);
 	}
 
-	for (TMap<FWorldPartitionReference, AActor*>::TIterator DirtyActorIt(DirtyActors); DirtyActorIt; ++DirtyActorIt)
+	for (TMap<FWorldPartitionReference, TObjectPtr<AActor>>::TIterator DirtyActorIt(DirtyActors); DirtyActorIt; ++DirtyActorIt)
 	{
 		if (!DirtyActorIt.Key().IsValid())
 		{
@@ -1937,7 +1937,7 @@ bool UWorldPartition::UnregisterActorDescContainer(UActorDescContainer* InActorD
 			{
 				ActorGuids.Add(It->GetGuid());
 
-				for (TMap<FWorldPartitionReference, AActor*>::TIterator DirtyActorIt(DirtyActors); DirtyActorIt; ++DirtyActorIt)
+				for (TMap<FWorldPartitionReference, TObjectPtr<AActor>>::TIterator DirtyActorIt(DirtyActors); DirtyActorIt; ++DirtyActorIt)
 				{
 					if (DirtyActorIt.Key() == ActorHandle)
 					{
diff --git a/Engine/Source/Runtime/Engine/Public/WorldPartition/WorldPartition.h b/Engine/Source/Runtime/Engine/Public/WorldPartition/WorldPartition.h
--- a/Engine/Source/Runtime/Engine/Public/WorldPartition/WorldPartition.h
+++ b/Engine/Source/Runtime/Engine/Public/WorldPartition/WorldPartition.h
@@ -314,7 +314,7 @@ public:
 
 	bool IsEnablingStreamingJustified() const { return bEnablingStreamingJustified; }
 
-	const TMap<FWorldPartitionReference, AActor*>& GetDirtyActors() const { return DirtyActors; }
+	const TMap<FWorldPartitionReference, AActor*>& GetDirtyActors() const { return ObjectPtrDecay(DirtyActors); }
 #endif
 
 public:
@@ -433,7 +433,7 @@ private:
 
 	TArray<FWorldPartitionReference> LoadedSubobjects;
 
-	TMap<FWorldPartitionReference, AActor*> DirtyActors;
+	TMap<FWorldPartitionReference, TObjectPtr<AActor>> DirtyActors;
 
 	TSet<FString> GeneratedStreamingPackageNames;
 
@@ -535,4 +535,4 @@ private:
 
 	friend class AWorldPartitionReplay;
 	friend class UWorldPartitionSubsystem;
-};
\ No newline at end of file
+};
diff --git a/Engine/Source/Runtime/Experimental/Iris/Core/Private/Iris/ReplicationSystem/NetRefHandleManager.h b/Engine/Source/Runtime/Experimental/Iris/Core/Private/Iris/ReplicationSystem/NetRefHandleManager.h
--- a/Engine/Source/Runtime/Experimental/Iris/Core/Private/Iris/ReplicationSystem/NetRefHandleManager.h
+++ b/Engine/Source/Runtime/Experimental/Iris/Core/Private/Iris/ReplicationSystem/NetRefHandleManager.h
@@ -358,7 +358,7 @@ private:
 	TArray<uint16> ReplicatedObjectRefCount;
 
 	// Raw pointers to all bound instances
-	TArray<UObject*> ReplicatedInstances;
+	TArray<TObjectPtr<UObject>> ReplicatedInstances;
 
 	// Assign handles
 	uint64 NextStaticHandleIndex;
diff --git a/Engine/Source/Runtime/Experimental/Iris/Core/Private/Iris/ReplicationSystem/ObjectReferenceCache.h b/Engine/Source/Runtime/Experimental/Iris/Core/Private/Iris/ReplicationSystem/ObjectReferenceCache.h
--- a/Engine/Source/Runtime/Experimental/Iris/Core/Private/Iris/ReplicationSystem/ObjectReferenceCache.h
+++ b/Engine/Source/Runtime/Experimental/Iris/Core/Private/Iris/ReplicationSystem/ObjectReferenceCache.h
@@ -156,7 +156,7 @@ private:
 
 	struct FQueuedBatchObjectReference
 	{
-		const UObject* Object = nullptr;
+		TObjectPtr<const UObject> Object = nullptr;
 		uint32 RefCount = 0U;
 	};
 
