Schedule external persistent levels for reload after initial save to make sure we can capture all dependent actors.

#rnx

diff --git a/Engine/Plugins/Developer/Concert/ConcertSync/ConcertSyncClient/Source/ConcertSyncClient/Private/ConcertClientPackageManager.cpp b/Engine/Plugins/Developer/Concert/ConcertSync/ConcertSyncClient/Source/ConcertSyncClient/Private/ConcertClientPackageManager.cpp
--- a/Engine/Plugins/Developer/Concert/ConcertSync/ConcertSyncClient/Source/ConcertSyncClient/Private/ConcertClientPackageManager.cpp
+++ b/Engine/Plugins/Developer/Concert/ConcertSync/ConcertSyncClient/Source/ConcertSyncClient/Private/ConcertClientPackageManager.cpp
@@ -125,17 +125,13 @@ FConcertClientPackageManager::~FConcertClientPackageManager()
 	// reloaded as long as the package hasn't been marked for purge. If it has
 	// been marked for purge, then it must have only existed in the sandbox
 	// that was just discarded, so it should not be reloaded.
-	UWorld* CurrentWorld = ConcertSyncClientUtil::GetCurrentWorld();
-	if (CurrentWorld)
+	ULevel* PersistentLevel = ConcertSyncClientUtil::GetExternalPersistentWorld();
+	if (PersistentLevel)
 	{
-		ULevel* PersistentLevel = CurrentWorld->PersistentLevel;
-		if (PersistentLevel && PersistentLevel->IsUsingExternalObjects())
+		const FName PackageName = PersistentLevel->GetPackage()->GetFName();
+		if (!PackagesPendingPurge.Contains(PackageName) && !PackagesPendingHotReload.Contains(PackageName))
 		{
-			const FName PackageName = PersistentLevel->GetPackage()->GetFName();
-			if (!PackagesPendingPurge.Contains(PackageName) && !PackagesPendingHotReload.Contains(PackageName))
-			{
-				PackagesPendingHotReload.Add(PackageName);
-			}
+			PackagesPendingHotReload.Add(PackageName);
 		}
 	}
 
@@ -464,6 +460,25 @@ void FConcertClientPackageManager::HandlePackageDirtyStateChanged(UPackage* InPa
 	}
 }
 
+void FConcertClientPackageManager::AddPendingReloadForNewExternalMaps(const FConcertPackageInfo& PackageInfo, const FString& PackagePathname)
+{
+	// Saved world partitioned main levels depend on loading other external assets. We need to schedule a reload of the map after initial
+	// save to ensure we have all of the content to fully load the p-level.
+	//
+	ULevel* PersistentLevel = ConcertSyncClientUtil::GetExternalPersistentWorld();
+	if (PersistentLevel)
+	{
+		const FName PackageName = PersistentLevel->GetPackage()->GetFName();
+		if (PackageName == PackageInfo.PackageName
+			&& !PackagesPendingPurge.Contains(PackageName)
+			&& !PackagesPendingHotReload.Contains(PackageName))
+		{
+			UE_LOG(LogConcert, Display, TEXT("Scheduling reloading for world partition persistent level %s."), *PackagePathname);
+			PackagesPendingHotReload.Add(PackageName);
+		}
+	}
+}
+
 void FConcertClientPackageManager::HandleLocalPackageEvent(const FConcertPackageInfo& PackageInfo, const FString& PackagePathname)
 {
 	// Ignore unwanted saves
@@ -491,6 +506,7 @@ void FConcertClientPackageManager::HandleLocalPackageEvent(const FConcertPackage
 				return; // Prevent capturing the original package state at every pre-save.
 			}
 		}
+		AddPendingReloadForNewExternalMaps(PackageInfo, PackagePathname);
 	}
 
 	if (PackageInfo.PackageUpdateType == EConcertPackageUpdateType::Added && EnumHasAnyFlags(LiveSession->GetSessionFlags(), EConcertSyncSessionFlags::ShouldUsePackageSandbox))
diff --git a/Engine/Plugins/Developer/Concert/ConcertSync/ConcertSyncClient/Source/ConcertSyncClient/Private/ConcertClientPackageManager.h b/Engine/Plugins/Developer/Concert/ConcertSync/ConcertSyncClient/Source/ConcertSyncClient/Private/ConcertClientPackageManager.h
--- a/Engine/Plugins/Developer/Concert/ConcertSync/ConcertSyncClient/Source/ConcertSyncClient/Private/ConcertClientPackageManager.h
+++ b/Engine/Plugins/Developer/Concert/ConcertSync/ConcertSyncClient/Source/ConcertSyncClient/Private/ConcertClientPackageManager.h
@@ -177,6 +177,11 @@ private:
 	 */
 	bool CanExchangePackageDataAsByteArray(int64 PackageDataSize) const;
 
+	/**
+	 * World partition nodes need to perform a map reload on initial save but only after all actors have been saved.
+	 */
+	void AddPendingReloadForNewExternalMaps(const FConcertPackageInfo& PackageInfo, const FString& PackagePathname);
+
 #if WITH_EDITOR
 	/**
 	 * Sandbox for storing package changes to disk within a Concert session.
diff --git a/Engine/Plugins/Developer/Concert/ConcertSync/ConcertSyncClient/Source/ConcertSyncClient/Private/ConcertSyncClientUtil.cpp b/Engine/Plugins/Developer/Concert/ConcertSync/ConcertSyncClient/Source/ConcertSyncClient/Private/ConcertSyncClientUtil.cpp
--- a/Engine/Plugins/Developer/Concert/ConcertSync/ConcertSyncClient/Source/ConcertSyncClient/Private/ConcertSyncClientUtil.cpp
+++ b/Engine/Plugins/Developer/Concert/ConcertSync/ConcertSyncClient/Source/ConcertSyncClient/Private/ConcertSyncClientUtil.cpp
@@ -725,6 +725,20 @@ UWorld* GetCurrentWorld()
 	return CurrentWorld;
 }
 
+ULevel* GetExternalPersistentWorld()
+{
+	UWorld* CurrentWorld = ConcertSyncClientUtil::GetCurrentWorld();
+	if (CurrentWorld)
+	{
+		ULevel* PersistentLevel = CurrentWorld->PersistentLevel;
+		if (PersistentLevel && PersistentLevel->IsUsingExternalObjects())
+		{
+			return PersistentLevel;
+		}
+	}
+	return nullptr;
+}
+
 void FillPackageInfo(UPackage* InPackage, UObject* InAsset, const EConcertPackageUpdateType InPackageUpdateType, FConcertPackageInfo& OutPackageInfo)
 {
 	UObject* Asset = InAsset ? InAsset : InPackage->FindAssetInPackage();
diff --git a/Engine/Plugins/Developer/Concert/ConcertSync/ConcertSyncClient/Source/ConcertSyncClient/Private/ConcertSyncClientUtil.h b/Engine/Plugins/Developer/Concert/ConcertSync/ConcertSyncClient/Source/ConcertSyncClient/Private/ConcertSyncClientUtil.h
--- a/Engine/Plugins/Developer/Concert/ConcertSync/ConcertSyncClient/Source/ConcertSyncClient/Private/ConcertSyncClientUtil.h
+++ b/Engine/Plugins/Developer/Concert/ConcertSync/ConcertSyncClient/Source/ConcertSyncClient/Private/ConcertSyncClientUtil.h
@@ -140,5 +140,8 @@ namespace ConcertSyncClientUtil
 
 	UWorld* GetCurrentWorld();
 
+	/** Return the Persistent Level if it is currently using external objects, i.e. one file per actor for world partitioning. */
+	ULevel* GetExternalPersistentWorld();
+
 	void FillPackageInfo(UPackage* InPackage, UObject* InAsset, const EConcertPackageUpdateType InPackageUpdateType, FConcertPackageInfo& OutPackageInfo);
 }
