Serialize pending take data and send it over the wire when users open the pending take. This allows users to capture uses of level sequences that were opened in Take Recorder.

#rnx

diff --git a/Engine/Plugins/Developer/Concert/ConcertSync/ConcertSyncClient/Source/ConcertSyncClient/Private/ConcertClientSequencerManager.cpp b/Engine/Plugins/Developer/Concert/ConcertSync/ConcertSyncClient/Source/ConcertSyncClient/Private/ConcertClientSequencerManager.cpp
--- a/Engine/Plugins/Developer/Concert/ConcertSync/ConcertSyncClient/Source/ConcertSyncClient/Private/ConcertClientSequencerManager.cpp
+++ b/Engine/Plugins/Developer/Concert/ConcertSync/ConcertSyncClient/Source/ConcertSyncClient/Private/ConcertClientSequencerManager.cpp
@@ -15,7 +15,7 @@
 #include "IConcertSession.h"
 #include "ConcertSettings.h"
 #include "ConcertClientSettings.h"
-
+#include "ConcertSyncArchives.h"
 #include "ConcertClientWorkspace.h"
 
 #include "Engine/GameEngine.h"
@@ -119,6 +119,16 @@ FConcertClientSequencerManager::FConcertClientSequencerManager(IConcertSyncClien
 	FCoreUObjectDelegates::OnPackageReloaded.AddRaw(this, &FConcertClientSequencerManager::HandleAssetReload);
 }
 
+namespace UE::Private::ConcertClientSequencerManager
+{
+const FString PendingTakePath = TEXT("/Temp/TakeRecorder/PendingTake.PendingTake:PendingTake");
+
+bool IsPendingTakePath(const FString& InSequencePath)
+{
+	return InSequencePath == PendingTakePath;
+}
+}
+
 FConcertClientSequencerManager::~FConcertClientSequencerManager()
 {
 	ISequencerModule* SequencerModulePtr = FModuleManager::Get().GetModulePtr<ISequencerModule>("Sequencer");
@@ -187,6 +197,25 @@ void FConcertClientSequencerManager::OnSequencerCreated(TSharedRef<ISequencer> I
 				FConcertSequencerOpenEvent OpenEvent;
 				OpenEvent.SequenceObjectPath = SequenceObjectPath;
 
+				if (UE::Private::ConcertClientSequencerManager::IsPendingTakePath(SequenceObjectPath))
+				{
+					// The pending take may have a level sequence loaded into it.  So we have to capture it with the object writer
+					// and transmit in our open event message.
+					//
+					ULevelSequence* LevelSequence = Cast<ULevelSequence>(Sequence);
+
+					FConcertSyncRemapObjectPath RemapperDelegate;
+					RemapperDelegate.BindLambda([](FString& InPath)
+					{
+						if (UE::Private::ConcertClientSequencerManager::IsPendingTakePath(InPath))
+						{
+							InPath = TEXT("/Engine/Transient.__PendingLevelSequence__");
+						}
+					});
+					FConcertSyncObjectWriter SyncObjectWriter(nullptr, LevelSequence, OpenEvent.TakeData.Bytes, true, false, RemapperDelegate);
+					SyncObjectWriter.SetSerializeNestedObjects(true);
+					SyncObjectWriter.SerializeObject(LevelSequence);
+				}
 				UE_LOG(LogConcertSequencerSync, Verbose, TEXT("    Sending OpenEvent: %s"), *OpenEvent.SequenceObjectPath);
 				Session->SendCustomEvent(OpenEvent, Session->GetSessionServerEndpointId(), EConcertMessageFlags::ReliableOrdered);
 			}
@@ -498,7 +527,8 @@ void FConcertClientSequencerManager::OnStateSyncEvent(const FConcertSessionConte
 		if (!bFoundSequencerForObject)
 		{
 			UE_LOG(LogConcertSequencerSync, VeryVerbose, TEXT("    No existing Sequencer with sequence %s open. Will open and sync a new one at end of frame."), *SequencerState.SequenceObjectPath);
-			PendingSequenceOpenEvents.Add(SequencerState.SequenceObjectPath);
+			FConcertSequencerOpenEvent OpenEvent{SequencerState.SequenceObjectPath, SequencerState.TakeData};
+			PendingSequenceOpenEvents.Add(MoveTemp(OpenEvent));
 			PendingSequencerEvents.Add(SequencerState);
 		}
 	}
@@ -573,7 +603,7 @@ void FConcertClientSequencerManager::OnCloseEvent(const FConcertSessionContext&,
 void FConcertClientSequencerManager::OnOpenEvent(const FConcertSessionContext&, const FConcertSequencerOpenEvent& InEvent)
 {
 	UE_LOG(LogConcertSequencerSync, Verbose, TEXT("Event Received - Open: %s. Deferring until end of frame."), *InEvent.SequenceObjectPath);
-	PendingSequenceOpenEvents.Add(InEvent.SequenceObjectPath);
+	PendingSequenceOpenEvents.Add(InEvent);
 }
 
 void FConcertClientSequencerManager::HandleAssetReload(const EPackageReloadPhase InPackageReloadPhase, FPackageReloadedEvent* InPackageReloadedEvent)
@@ -683,22 +713,53 @@ void FConcertClientSequencerManager::ApplyCloseEvent(const FConcertSequencerClos
 	ApplyCloseEventToPlayers(CloseEvent);
 }
 
-void FConcertClientSequencerManager::ApplyOpenEvent(const FString& SequenceObjectPath)
+void FConcertClientSequencerManager::ApplyOpenEvent(const FConcertSequencerOpenEvent& InOpenEvent)
 {
 	TGuardValue<bool> ReentrancyGuard(bRespondingToTransportEvent, true);
-
+	FString SequenceObjectPath = InOpenEvent.SequenceObjectPath;
 	UE_LOG(LogConcertSequencerSync, Verbose, TEXT("Handling Event - Open: %s"), *SequenceObjectPath);
 
+	bool bDidOpen = false;
 #if WITH_EDITOR
 	if (IsSequencerRemoteOpenEnabled() && GIsEditor)
 	{
-		GEditor->GetEditorSubsystem<UAssetEditorSubsystem>()->OpenEditorForAsset(SequenceObjectPath);
+		if (!UE::Private::ConcertClientSequencerManager::IsPendingTakePath(SequenceObjectPath))
+		{
+			// Don't open the asset until we have loaded in the pending take data.
+			GEditor->GetEditorSubsystem<UAssetEditorSubsystem>()->OpenEditorForAsset(SequenceObjectPath);
+		}
+
+		bDidOpen = true;
 	}
 
 #endif
 	if (!GIsEditor && CVarEnableSequencePlayer.GetValueOnAnyThread() > 0)
 	{
 		CreateNewSequencePlayerIfNotExists(*SequenceObjectPath);
+		bDidOpen = true;
+	}
+
+	if (bDidOpen && UE::Private::ConcertClientSequencerManager::IsPendingTakePath(SequenceObjectPath))
+	{
+		ULevelSequence* PendingLevelSequence = FindObject<ULevelSequence>(nullptr, *UE::Private::ConcertClientSequencerManager::PendingTakePath);
+
+		// Apply pending take data to our version of the pending take.
+		//
+		FConcertSyncEncounteredMissingObject MissingObjectDelegate;
+		MissingObjectDelegate.BindLambda([](const FStringView MissingObject)
+			{
+				UE_LOG(LogConcertSequencerSync, Display, TEXT("Missing Object %s when loading PendingTake"), MissingObject);
+			});
+
+		FConcertSyncWorldRemapper Remapper(
+			TEXT("/Engine/Transient.__PendingLevelSequence__"), PendingLevelSequence->GetPathName());
+		FConcertSyncObjectReader Reader(nullptr, MoveTemp(Remapper), nullptr, PendingLevelSequence, InOpenEvent.TakeData.Bytes,
+										MissingObjectDelegate);
+		Reader.SerializeObject(PendingLevelSequence);
+		#if WITH_EDITOR
+		GEditor->GetEditorSubsystem<UAssetEditorSubsystem>()->OpenEditorForAsset(SequenceObjectPath);
+		#endif
+		PendingLevelSequence->GetPackage()->SetDirtyFlag(false);
 	}
 }
 
@@ -1242,9 +1303,9 @@ void FConcertClientSequencerManager::OnWorkspaceEndFrameCompleted()
 	}
 	PendingCreate.Reset();
 
-	for (const FString& SequenceObjectPath : PendingSequenceOpenEvents)
+	for (const FConcertSequencerOpenEvent& SequenceOpenEvent : PendingSequenceOpenEvents)
 	{
-		ApplyOpenEvent(SequenceObjectPath);
+		ApplyOpenEvent(SequenceOpenEvent);
 	}
 	PendingSequenceOpenEvents.Reset();
 
diff --git a/Engine/Plugins/Developer/Concert/ConcertSync/ConcertSyncClient/Source/ConcertSyncClient/Private/ConcertClientSequencerManager.h b/Engine/Plugins/Developer/Concert/ConcertSync/ConcertSyncClient/Source/ConcertSyncClient/Private/ConcertClientSequencerManager.h
--- a/Engine/Plugins/Developer/Concert/ConcertSync/ConcertSyncClient/Source/ConcertSyncClient/Private/ConcertClientSequencerManager.h
+++ b/Engine/Plugins/Developer/Concert/ConcertSync/ConcertSyncClient/Source/ConcertSyncClient/Private/ConcertClientSequencerManager.h
@@ -228,7 +228,7 @@ private:
 	 *
 	 * @param SequenceObjectPath	The sequence to open
 	 */
-	void ApplyOpenEvent(const FString& SequenceObjectPath);
+	void ApplyOpenEvent(const FConcertSequencerOpenEvent& InOpenEvent);
 
 	/**
 	 * Create a new sequence player to be used in -game instances.
@@ -329,7 +329,7 @@ private:
 	TArray<FConcertSequencerState> PendingSequencerEvents;
 
 	/** List of pending sequencer open events to apply at end of frame. */
-	TArray<FString> PendingSequenceOpenEvents;
+	TArray<FConcertSequencerOpenEvent> PendingSequenceOpenEvents;
 
 	/** List of pending sequencer open events to apply at end of frame. */
 	TArray<FConcertSequencerCloseEvent> PendingSequenceCloseEvents;
diff --git a/Engine/Plugins/Developer/Concert/ConcertSync/ConcertSyncCore/Source/ConcertSyncCore/Private/ConcertSyncArchives.cpp b/Engine/Plugins/Developer/Concert/ConcertSync/ConcertSyncCore/Source/ConcertSyncCore/Private/ConcertSyncArchives.cpp
--- a/Engine/Plugins/Developer/Concert/ConcertSync/ConcertSyncCore/Source/ConcertSyncCore/Private/ConcertSyncArchives.cpp
+++ b/Engine/Plugins/Developer/Concert/ConcertSync/ConcertSyncCore/Source/ConcertSyncCore/Private/ConcertSyncArchives.cpp
@@ -238,8 +238,27 @@ FConcertSyncObjectWriter::FConcertSyncObjectWriter(FConcertLocalIdentifierTable*
 {
 }
 
-void FConcertSyncObjectWriter::SerializeObject(const UObject* InObject, const TArray<const FProperty*>* InPropertiesToWrite)
+void FConcertSyncObjectWriter::SerializeObject(const UObject* InObject, const TArray<const FProperty*>* InPropertiesToWrite, bool bAllowOuters /*= false*/)
 {
+	PackageName = InObject->GetPackage()->GetPathName();
+
+	if (bSerializeNestedObjects && !bAllowOuters)
+	{
+		// When serializing nested objects, we need to know when to stop
+		// capturing. A property can point to outer or this package and those
+		// should not be included in the serialization. So we make note of them
+		// here to stop the serialization when those objects are encountered and
+		// only store the FSoftObjectPath
+		CollectedObjects.Add(InObject->GetPackage()->GetPathName());
+		CollectedObjects.Add(InObject->GetPathName());
+		UObject* Outer = InObject->GetOuter();
+		while (	Outer )
+		{
+			CollectedObjects.Add(Outer->GetPathName());
+			Outer = Outer->GetOuter();
+		}
+	}
+
 	if (InPropertiesToWrite)
 	{
 		ShouldSkipPropertyFunc = [InObject, InPropertiesToWrite](const FProperty* InProperty) -> bool
@@ -267,6 +286,7 @@ void FConcertSyncObjectWriter::SerializeProperty(const FProperty* InProp, const
 
 FArchive& FConcertSyncObjectWriter::operator<<(UObject*& Obj)
 {
+	bool bCanSerializeObject = false;
 	FSoftObjectPath ObjPath;
 	if (Obj)
 	{
@@ -277,11 +297,34 @@ FArchive& FConcertSyncObjectWriter::operator<<(UObject*& Obj)
 		else
 		{
 			FString ObjectPathString = Obj->GetPathName();
+
+			// Check to see if we can serialize before we perform the remapping.
+			bCanSerializeObject = !CollectedObjects.Contains(ObjectPathString);
+
 			RemapObjectPathDelegate.ExecuteIfBound(ObjectPathString);
 			ObjPath.SetPath(ObjectPathString);
 		}
 	}
 	ObjPath.SerializePath(*this);
+
+	bool bDoSerializeObject = bCanSerializeObject
+		&& bSerializeNestedObjects
+		&& PackageName == Obj->GetPackage()->GetPathName()
+		&& !Obj->IsA<UPackage>();
+
+	*this << bDoSerializeObject;
+	if (bDoSerializeObject)
+	{
+		// To prevent stack overflow we need to add to the collected objects first.
+		CollectedObjects.Add(ObjPath.ToString());
+
+		FString ClassName = Obj->GetClass()->GetPathName();
+		int32 OutFlags = static_cast<int32>(Obj->GetFlags());
+		*this << ClassName;
+		*this << OutFlags;
+
+		SerializeObject(Obj, nullptr, true);
+	}
 	return *this;
 }
 
@@ -400,6 +443,7 @@ FConcertSyncObjectReader::FConcertSyncObjectReader(const FConcertLocalIdentifier
 
 void FConcertSyncObjectReader::SerializeObject(UObject* InObject)
 {
+	CurrentOuter = InObject;
 	InObject->Serialize(*this);
 }
 
@@ -416,6 +460,8 @@ FArchive& FConcertSyncObjectReader::operator<<(UObject*& Obj)
 	FSoftObjectPath ObjPath;
 	ObjPath.SerializePath(*this);
 
+	bool bDidSerializeObject;
+	*this << bDidSerializeObject;
 	if (ObjPath.IsNull())
 	{
 		Obj = nullptr;
@@ -452,6 +498,25 @@ FArchive& FConcertSyncObjectReader::operator<<(UObject*& Obj)
 				EncounteredMissingObjectDelegate.Execute(FStringView(ResolvedObjPath));
 			}
 		}
+
+		if (bDidSerializeObject)
+		{
+			FString ClassName;
+			int32  CapturedObjectCreationFlags;
+			*this << ClassName;
+			*this << CapturedObjectCreationFlags;
+			EObjectFlags ObjectCreationFlags = static_cast<EObjectFlags>(CapturedObjectCreationFlags);
+
+			if (!Obj)
+			{
+				UClass* InClass = FindObject<UClass>(nullptr, *ClassName);
+				check(InClass);
+				Obj = NewObject<UObject>(CurrentOuter, InClass, NAME_None, ObjectCreationFlags);
+			}
+
+			check(Obj);
+			SerializeObject(Obj);
+		}
 	}
 
 	return *this;
diff --git a/Engine/Plugins/Developer/Concert/ConcertSync/ConcertSyncCore/Source/ConcertSyncCore/Public/ConcertSequencerMessages.h b/Engine/Plugins/Developer/Concert/ConcertSync/ConcertSyncCore/Source/ConcertSyncCore/Public/ConcertSequencerMessages.h
--- a/Engine/Plugins/Developer/Concert/ConcertSync/ConcertSyncCore/Source/ConcertSyncCore/Public/ConcertSequencerMessages.h
+++ b/Engine/Plugins/Developer/Concert/ConcertSync/ConcertSyncCore/Source/ConcertSyncCore/Public/ConcertSequencerMessages.h
@@ -3,6 +3,7 @@
 #pragma once
 
 #include "Misc/QualifiedFrameTime.h"
+#include "ConcertMessageData.h"
 #include "ConcertSequencerMessages.generated.h"
 
 
@@ -46,9 +47,18 @@ struct FConcertSequencerState
 	UPROPERTY()
 	float PlaybackSpeed;
 
+	/** Indicate if we are currently looping */
 	UPROPERTY()
 	bool bLoopMode = false;
 
+	/**
+	 * In the case that the SequenceObjectPath points to a take preset, we capture the preset data
+	 * into a payload that can be applied to take that we are going to open. We store it in the state
+	 * so that we can play it back when new users join.
+	 */
+	UPROPERTY()
+	FConcertByteArray TakeData;
+
 	FConcertSequencerState()
 		: PlayerStatus(EConcertMovieScenePlayerStatus::Stopped)
 		, PlaybackSpeed(1.0f)
@@ -66,6 +76,14 @@ struct FConcertSequencerOpenEvent
 	/** The full path name to the root sequence of the sequencer that just opened. */
 	UPROPERTY()
 	FString SequenceObjectPath;
+
+	/**
+	 * In the case that the SequenceObjectPath points to a take preset, we capture the preset data
+	 * into a payload that can be applied to take that we are going to open. We store it in the state
+	 * so that we can play it back when new users join.
+	 */
+	UPROPERTY()
+	FConcertByteArray TakeData;
 };
 
 /**
diff --git a/Engine/Plugins/Developer/Concert/ConcertSync/ConcertSyncCore/Source/ConcertSyncCore/Public/ConcertSyncArchives.h b/Engine/Plugins/Developer/Concert/ConcertSync/ConcertSyncCore/Source/ConcertSyncCore/Public/ConcertSyncArchives.h
--- a/Engine/Plugins/Developer/Concert/ConcertSync/ConcertSyncCore/Source/ConcertSyncCore/Public/ConcertSyncArchives.h
+++ b/Engine/Plugins/Developer/Concert/ConcertSync/ConcertSyncCore/Source/ConcertSyncCore/Public/ConcertSyncArchives.h
@@ -53,10 +53,11 @@ private:
 class CONCERTSYNCCORE_API FConcertSyncObjectWriter : public FConcertIdentifierWriter
 {
 public:
+
 	FConcertSyncObjectWriter(FConcertLocalIdentifierTable* InLocalIdentifierTable, UObject* InObj, TArray<uint8>& OutBytes, const bool InIncludeEditorOnlyData, const bool InSkipAssets, const FConcertSyncRemapObjectPath& InRemapDelegate);
 	FConcertSyncObjectWriter(FConcertLocalIdentifierTable* InLocalIdentifierTable, UObject* InObj, TArray<uint8>& OutBytes, const bool InIncludeEditorOnlyData, const bool InSkipAssets);
 
-	void SerializeObject(const UObject* InObject, const TArray<const FProperty*>* InPropertiesToWrite = nullptr);
+	void SerializeObject(const UObject* InObject, const TArray<const FProperty*>* InPropertiesToWrite = nullptr, bool bAllowOuters = false);
 	void SerializeProperty(const FProperty* InProp, const UObject* InObject);
 
 	using FConcertIdentifierWriter::operator<<; // For visibility of the overloads we don't override
@@ -72,11 +73,18 @@ public:
 	virtual bool ShouldSkipProperty(const FProperty* InProperty) const override;
 	//~ End FArchive Interface
 
+	void SetSerializeNestedObjects(bool bInSerializeNestedObjects)
+	{
+		bSerializeNestedObjects = bInSerializeNestedObjects;
+	}
 private:
 	typedef TFunction<bool(const FProperty*)> FShouldSkipPropertyFunc;
 	bool bSkipAssets;
+	bool bSerializeNestedObjects = false;
+	FString PackageName;
 	FShouldSkipPropertyFunc ShouldSkipPropertyFunc;
 	FConcertSyncRemapObjectPath RemapObjectPathDelegate;
+	TSet<FString> CollectedObjects;
 };
 
 /** Archive for reading objects that have been received from another instance via Concert */
@@ -105,6 +113,7 @@ protected:
 	virtual void OnObjectSerialized(const FSoftObjectPath& Obj) {}
 
 private:
+	UObject* CurrentOuter = nullptr;
 	FConcertSyncWorldRemapper WorldRemapper;
 	FConcertSyncEncounteredMissingObject EncounteredMissingObjectDelegate;
 };
diff --git a/Engine/Plugins/Developer/Concert/ConcertSync/ConcertSyncServer/Source/ConcertSyncServer/Private/ConcertServerSequencerManager.cpp b/Engine/Plugins/Developer/Concert/ConcertSync/ConcertSyncServer/Source/ConcertSyncServer/Private/ConcertServerSequencerManager.cpp
--- a/Engine/Plugins/Developer/Concert/ConcertSync/ConcertSyncServer/Source/ConcertSyncServer/Private/ConcertServerSequencerManager.cpp
+++ b/Engine/Plugins/Developer/Concert/ConcertSync/ConcertSyncServer/Source/ConcertSyncServer/Private/ConcertServerSequencerManager.cpp
@@ -73,6 +73,7 @@ void FConcertServerSequencerManager::HandleSequencerOpenEvent(const FConcertSess
 	FConcertOpenSequencerState& SequencerState = SequencerStates.FindOrAdd(*InEvent.SequenceObjectPath);
 	SequencerState.ClientEndpointIds.AddUnique(InEventContext.SourceEndpointId);
 	SequencerState.State.SequenceObjectPath = InEvent.SequenceObjectPath;
+	SequencerState.TakeData = InEvent.TakeData;
 
 	// Forward the message to the other clients
 	TArray<FGuid> ClientIds = LiveSession->GetSession().GetSessionClientEndpointIds();
@@ -183,7 +184,9 @@ void FConcertServerSequencerManager::HandleWorkspaceSyncAndFinalizeCompletedEven
 	FConcertSequencerStateSyncEvent SequencerStateSyncEvent;
 	for (const auto& Pair : SequencerStates)
 	{
-		SequencerStateSyncEvent.SequencerStates.Add(Pair.Value.State);
+		FConcertSequencerState OutState = Pair.Value.State;
+		OutState.TakeData = Pair.Value.TakeData;
+		SequencerStateSyncEvent.SequencerStates.Emplace(MoveTemp(OutState));
 	}
 
 	LiveSession->GetSession().SendCustomEvent(SequencerStateSyncEvent, InEventContext.SourceEndpointId, EConcertMessageFlags::ReliableOrdered);
diff --git a/Engine/Plugins/Developer/Concert/ConcertSync/ConcertSyncServer/Source/ConcertSyncServer/Private/ConcertServerSequencerManager.h b/Engine/Plugins/Developer/Concert/ConcertSync/ConcertSyncServer/Source/ConcertSyncServer/Private/ConcertServerSequencerManager.h
--- a/Engine/Plugins/Developer/Concert/ConcertSync/ConcertSyncServer/Source/ConcertSyncServer/Private/ConcertServerSequencerManager.h
+++ b/Engine/Plugins/Developer/Concert/ConcertSync/ConcertSyncServer/Source/ConcertSyncServer/Private/ConcertServerSequencerManager.h
@@ -63,6 +63,14 @@ private:
 
 		/** Current state of the Sequence. */
 		FConcertSequencerState State;
+
+		/**
+		 * In the case that the SequenceObjectPath points to a take preset. We capture the preset data
+		 * into a payload that can be applied to take that we are going to open. We store it in the state
+		 * so that we can play it back when new users join.
+		 */
+		UPROPERTY()
+		FConcertByteArray TakeData;
 	};
 
 	/** Map of all currently opened Sequencer in a session, locally opened or not. */
