Only call the base implementation to PreEditChange() when in `-game` nodes to prevent FlushRenderingCommands().  This is only for interactive changes.

diff --git a/Engine/Plugins/Developer/Concert/ConcertSync/ConcertSyncClient/Source/ConcertSyncClient/Private/ConcertClientTransactionBridge.cpp b/Engine/Plugins/Developer/Concert/ConcertSync/ConcertSyncClient/Source/ConcertSyncClient/Private/ConcertClientTransactionBridge.cpp
--- a/Engine/Plugins/Developer/Concert/ConcertSync/ConcertSyncClient/Source/ConcertSyncClient/Private/ConcertClientTransactionBridge.cpp
+++ b/Engine/Plugins/Developer/Concert/ConcertSync/ConcertSyncClient/Source/ConcertSyncClient/Private/ConcertClientTransactionBridge.cpp
@@ -449,7 +449,16 @@ void ProcessTransactionEvent(const FConcertTransactionEventBase& InEvent, const
 			{
 				if (bIsSnapshot)
 				{
-					TransactionObject->PreEditChange(TransactionProp);
+					if (!GIsEditor)
+					{
+						// Only call the base implementation for -game nodes as we want to avoid FlushRenderingCommands()
+						// during interactive changes. 
+						Cast<UObject>(TransactionObject)->PreEditChange(TransactionProp);
+					}
+					else 
+					{
+						TransactionObject->PreEditChange(TransactionProp);
+					}
 				}
 
 				FEditPropertyChain PropertyChain;
