Add a scoped variable to control the call of FlushRenderingCommands() in ActorComponent.  This causes some performance issues on -game instances.  It is only disabled during snapshot transactions in Multi-user.

diff --git a/Engine/Plugins/Developer/Concert/ConcertSync/ConcertSyncClient/Source/ConcertSyncClient/Private/ConcertClientTransactionBridge.cpp b/Engine/Plugins/Developer/Concert/ConcertSync/ConcertSyncClient/Source/ConcertSyncClient/Private/ConcertClientTransactionBridge.cpp
--- a/Engine/Plugins/Developer/Concert/ConcertSync/ConcertSyncClient/Source/ConcertSyncClient/Private/ConcertClientTransactionBridge.cpp
+++ b/Engine/Plugins/Developer/Concert/ConcertSync/ConcertSyncClient/Source/ConcertSyncClient/Private/ConcertClientTransactionBridge.cpp
@@ -27,6 +27,8 @@
 LLM_DEFINE_TAG(Concert_ConcertClientTransactionBridge);
 #define LOCTEXT_NAMESPACE "ConcertClientTransactionBridge"
 
+extern ENGINE_API bool GFlushRenderingCommandsOnPreEditChange;
+
 namespace ConcertClientTransactionBridgeUtil
 {
 
@@ -265,6 +267,7 @@ struct FEditorTransactionNotification
 };
 #endif
 
+
 void ProcessTransactionEvent(const FConcertTransactionEventBase& InEvent, const FConcertSessionVersionInfo* InVersionInfo, const TArray<FName>& InPackagesToProcess, const FConcertLocalIdentifierTable* InLocalIdentifierTablePtr, const bool bIsSnapshot, const FConcertSyncWorldRemapper& WorldRemapper, const bool bIncludeEditorOnlyProperties)
 {
 	// Transactions are applied in multiple-phases...
@@ -449,16 +452,9 @@ void ProcessTransactionEvent(const FConcertTransactionEventBase& InEvent, const
 			{
 				if (bIsSnapshot)
 				{
-					if (!GIsEditor)
-					{
-						// Only call the base implementation for -game nodes as we want to avoid FlushRenderingCommands()
-						// during interactive changes. 
-						Cast<UObject>(TransactionObject)->PreEditChange(TransactionProp);
-					}
-					else 
-					{
-						TransactionObject->PreEditChange(TransactionProp);
-					}
+					// Prevent FlushRenderingCommands from running when in -game. It can cause performance issues / hitching during user interaction.
+					TGuardValue<bool> ShouldFlushRenderingCommands(GFlushRenderingCommandsOnPreEditChange, GIsEditor);					
+					TransactionObject->PreEditChange(TransactionProp);
 				}
 
 				FEditPropertyChain PropertyChain;
diff --git a/Engine/Source/Runtime/Engine/Private/Components/ActorComponent.cpp b/Engine/Source/Runtime/Engine/Private/Components/ActorComponent.cpp
--- a/Engine/Source/Runtime/Engine/Private/Components/ActorComponent.cpp
+++ b/Engine/Source/Runtime/Engine/Private/Components/ActorComponent.cpp
@@ -801,6 +801,8 @@ bool UActorComponent::Modify( bool bAlwaysMarkDirty/*=true*/ )
 	return Super::Modify(bAlwaysMarkDirty);
 }
 
+ENGINE_API bool GFlushRenderingCommandsOnPreEditChange = true;
+
 void UActorComponent::PreEditChange(FProperty* PropertyThatWillChange)
 {
 	Super::PreEditChange(PropertyThatWillChange);
@@ -825,9 +827,11 @@ void UActorComponent::PreEditChange(FProperty* PropertyThatWillChange)
 			WorldPrivate = nullptr;
 		}
 	}
-
 	// Flush rendering commands to ensure the rendering thread processes the component detachment before it is modified.
-	FlushRenderingCommands();
+	if (GFlushRenderingCommandsOnPreEditChange)
+	{
+		FlushRenderingCommands();
+	}
 }
 
 void UActorComponent::PreEditUndo()
