Cloase all open sequencers opened during a Multi-user session.

#rnx

diff --git a/Engine/Plugins/Developer/Concert/ConcertSync/ConcertSyncClient/Source/ConcertSyncClient/Private/ConcertClientSequencerManager.cpp b/Engine/Plugins/Developer/Concert/ConcertSync/ConcertSyncClient/Source/ConcertSyncClient/Private/ConcertClientSequencerManager.cpp
--- a/Engine/Plugins/Developer/Concert/ConcertSync/ConcertSyncClient/Source/ConcertSyncClient/Private/ConcertClientSequencerManager.cpp
+++ b/Engine/Plugins/Developer/Concert/ConcertSync/ConcertSyncClient/Source/ConcertSyncClient/Private/ConcertClientSequencerManager.cpp
@@ -284,12 +284,27 @@ void FConcertClientSequencerManager::Unregister(TSharedRef<IConcertClientSession
 	if (TSharedPtr<IConcertClientSession> Session = WeakSession.Pin())
 	{
 		check(Session == InSession);
+
 		Session->UnregisterCustomEventHandler<FConcertSequencerStateEvent>(this);
 		Session->UnregisterCustomEventHandler<FConcertSequencerCloseEvent>(this);
 		Session->UnregisterCustomEventHandler<FConcertSequencerOpenEvent>(this);
 		Session->UnregisterCustomEventHandler<FConcertSequencerStateSyncEvent>(this);
 		Session->UnregisterCustomEventHandler<FConcertSequencerTimeAdjustmentEvent>(this);
 		Session->UnregisterCustomEventHandler<FConcertSequencerPreloadRequest>(this);
+
+		if (GEditor)
+		{
+			for (FOpenSequencerData& OpenSequencer : OpenSequencers)
+			{
+				TSharedPtr<ISequencer> Sequencer = OpenSequencer.WeakSequencer.Pin();
+				UMovieSceneSequence* Sequence = Sequencer.IsValid() ? Sequencer->GetRootMovieSceneSequence() : nullptr;
+
+				if (Sequence)
+				{
+					GEditor->GetEditorSubsystem<UAssetEditorSubsystem>()->CloseAllEditorsForAsset(Sequence);
+				}
+			}
+		}
 	}
 
 	WeakSession.Reset();
