Fix ensure during leaving MU session.

#rnx

diff --git a/Engine/Plugins/Developer/Concert/ConcertSync/ConcertSyncClient/Source/ConcertSyncClient/Private/ConcertClientSequencerManager.cpp b/Engine/Plugins/Developer/Concert/ConcertSync/ConcertSyncClient/Source/ConcertSyncClient/Private/ConcertClientSequencerManager.cpp
--- a/Engine/Plugins/Developer/Concert/ConcertSync/ConcertSyncClient/Source/ConcertSyncClient/Private/ConcertClientSequencerManager.cpp
+++ b/Engine/Plugins/Developer/Concert/ConcertSync/ConcertSyncClient/Source/ConcertSyncClient/Private/ConcertClientSequencerManager.cpp
@@ -294,16 +294,20 @@ void FConcertClientSequencerManager::Unregister(TSharedRef<IConcertClientSession
 
 		if (GEditor)
 		{
+			TArray<UMovieSceneSequence*> Sequences;
 			for (FOpenSequencerData& OpenSequencer : OpenSequencers)
 			{
 				TSharedPtr<ISequencer> Sequencer = OpenSequencer.WeakSequencer.Pin();
-				UMovieSceneSequence* Sequence = Sequencer.IsValid() ? Sequencer->GetRootMovieSceneSequence() : nullptr;
-
-				if (Sequence)
+				if (Sequencer.IsValid() && Sequencer->GetRootMovieSceneSequence())
 				{
-					GEditor->GetEditorSubsystem<UAssetEditorSubsystem>()->CloseAllEditorsForAsset(Sequence);
+					Sequences.Add(Sequencer->GetRootMovieSceneSequence());
 				}
 			}
+
+			for (UMovieSceneSequence* MovieSequence : Sequences)
+			{
+				GEditor->GetEditorSubsystem<UAssetEditorSubsystem>()->CloseAllEditorsForAsset(MovieSequence);
+			}
 		}
 	}
 
