Fix crash on -game instances when loading the pending take.

#rnx

diff --git a/Engine/Plugins/Developer/Concert/ConcertSync/ConcertSyncClient/Source/ConcertSyncClient/Private/ConcertClientSequencerManager.cpp b/Engine/Plugins/Developer/Concert/ConcertSync/ConcertSyncClient/Source/ConcertSyncClient/Private/ConcertClientSequencerManager.cpp
--- a/Engine/Plugins/Developer/Concert/ConcertSync/ConcertSyncClient/Source/ConcertSyncClient/Private/ConcertClientSequencerManager.cpp
+++ b/Engine/Plugins/Developer/Concert/ConcertSync/ConcertSyncClient/Source/ConcertSyncClient/Private/ConcertClientSequencerManager.cpp
@@ -777,7 +777,10 @@ void FConcertClientSequencerManager::ApplyOpenEvent(const FConcertSequencerOpenE
 		Reader.SetSerializeNestedObjects(true);
 		Reader.SerializeObject(PendingLevelSequence);
 		#if WITH_EDITOR
-		GEditor->GetEditorSubsystem<UAssetEditorSubsystem>()->OpenEditorForAsset(SequenceObjectPath);
+		if (GIsEditor)
+		{
+			GEditor->GetEditorSubsystem<UAssetEditorSubsystem>()->OpenEditorForAsset(SequenceObjectPath);
+		}
 		#endif
 		PendingLevelSequence->GetPackage()->SetDirtyFlag(false);
 	}
