Fix Concert compile errors for UnreaGame

diff --git a/Engine/Plugins/Developer/Concert/ConcertSync/ConcertSyncClient/Source/ConcertSyncClient/Private/ConcertClientLockManager.cpp b/Engine/Plugins/Developer/Concert/ConcertSync/ConcertSyncClient/Source/ConcertSyncClient/Private/ConcertClientLockManager.cpp
--- a/Engine/Plugins/Developer/Concert/ConcertSync/ConcertSyncClient/Source/ConcertSyncClient/Private/ConcertClientLockManager.cpp
+++ b/Engine/Plugins/Developer/Concert/ConcertSync/ConcertSyncClient/Source/ConcertSyncClient/Private/ConcertClientLockManager.cpp
@@ -98,10 +98,12 @@ bool FConcertClientLockManager::CanSavePackage(UPackage* InPackage, const FStrin
 			return false;
 		}
 	}
+#if WITH_EDITOR
 	if (OkToSaveBackupDelegate.IsBound())
 	{
 		return OkToSaveBackupDelegate.Execute(InPackage, InFilename, ErrorLog);
 	}
+#endif
 	return true;
 }
 
diff --git a/Engine/Plugins/Developer/Concert/ConcertSync/ConcertSyncClient/Source/ConcertSyncClient/Private/ConcertClientPackageBridge.cpp b/Engine/Plugins/Developer/Concert/ConcertSync/ConcertSyncClient/Source/ConcertSyncClient/Private/ConcertClientPackageBridge.cpp
--- a/Engine/Plugins/Developer/Concert/ConcertSync/ConcertSyncClient/Source/ConcertSyncClient/Private/ConcertClientPackageBridge.cpp
+++ b/Engine/Plugins/Developer/Concert/ConcertSync/ConcertSyncClient/Source/ConcertSyncClient/Private/ConcertClientPackageBridge.cpp
@@ -25,6 +25,7 @@
 
 #if WITH_EDITOR
 	#include "LevelEditor.h"
+	#include "UnrealEdMisc.h"
 #endif
 
 LLM_DEFINE_TAG(Concert_ConcertClientPackageBridge);
@@ -136,6 +137,7 @@ bool& FConcertClientPackageBridge::GetIgnoreLocalDiscardRef()
 	return bIgnoreLocalDiscard;
 }
 
+#if WITH_EDITOR
 void FConcertClientPackageBridge::HandlePackagePreSave(UPackage* Package, FObjectPreSaveContext ObjectSaveContext)
 {
 	// Only execute if this is a user save
@@ -360,6 +362,7 @@ void FConcertClientPackageBridge::HandleMapChanged(UWorld* InWorld, EMapChangeTy
 		}
 	}
 }
+#endif
 
 void FConcertClientPackageBridge::RegisterPackageFilter(FName FilterName, FPackageFilterDelegate FilterHandle)
 {
diff --git a/Engine/Plugins/Developer/Concert/ConcertSync/ConcertSyncClient/Source/ConcertSyncClient/Private/ConcertClientPackageBridge.h b/Engine/Plugins/Developer/Concert/ConcertSync/ConcertSyncClient/Source/ConcertSyncClient/Private/ConcertClientPackageBridge.h
--- a/Engine/Plugins/Developer/Concert/ConcertSync/ConcertSyncClient/Source/ConcertSyncClient/Private/ConcertClientPackageBridge.h
+++ b/Engine/Plugins/Developer/Concert/ConcertSync/ConcertSyncClient/Source/ConcertSyncClient/Private/ConcertClientPackageBridge.h
@@ -32,6 +32,8 @@ public:
 	virtual EPackageFilterResult IsPackageFiltered(const FConcertPackageInfo& PackageInfo) const override;
 
 private:
+	
+#if WITH_EDITOR
 	/** Handle any deferred tasks that should be run on the game thread. */
 	void OnEndFrame();
 
@@ -55,6 +57,7 @@ private:
 
 	/** Called when the editor map is changed */
 	void HandleMapChanged(UWorld* InWorld, EMapChangeType InMapChangeType);
+#endif
 
 	/** Called when a local package event happens */
 	FOnConcertClientLocalPackageEvent OnLocalPackageEventDelegate;
diff --git a/Engine/Plugins/Developer/Concert/ConcertSync/ConcertSyncClient/Source/ConcertSyncClient/Private/ConcertClientPresenceManager.cpp b/Engine/Plugins/Developer/Concert/ConcertSync/ConcertSyncClient/Source/ConcertSyncClient/Private/ConcertClientPresenceManager.cpp
--- a/Engine/Plugins/Developer/Concert/ConcertSync/ConcertSyncClient/Source/ConcertSyncClient/Private/ConcertClientPresenceManager.cpp
+++ b/Engine/Plugins/Developer/Concert/ConcertSync/ConcertSyncClient/Source/ConcertSyncClient/Private/ConcertClientPresenceManager.cpp
@@ -985,12 +985,9 @@ void FConcertClientDefaultPresenceModeFactory::QueryVRDeviceType()
 
 #else
 
-namespace FConcertClientPresenceManager
+double FConcertClientPresenceManager::GetLocationUpdateFrequency()
 {
-	double GetLocationUpdateFrequency()
-	{
-		return ConcertClientPresenceManagerUtil::LocationUpdateFrequencySeconds;
-	}
+	return ConcertClientPresenceManagerUtil::LocationUpdateFrequencySeconds;
 }
 
 #endif // WITH_EDITOR
diff --git a/Engine/Plugins/Developer/Concert/ConcertSync/ConcertSyncClient/Source/ConcertSyncClient/Private/ConcertClientPresenceManager.h b/Engine/Plugins/Developer/Concert/ConcertSync/ConcertSyncClient/Source/ConcertSyncClient/Private/ConcertClientPresenceManager.h
--- a/Engine/Plugins/Developer/Concert/ConcertSync/ConcertSyncClient/Source/ConcertSyncClient/Private/ConcertClientPresenceManager.h
+++ b/Engine/Plugins/Developer/Concert/ConcertSync/ConcertSyncClient/Source/ConcertSyncClient/Private/ConcertClientPresenceManager.h
@@ -318,10 +318,10 @@ private:
 
 #else
 
-namespace FConcertClientPresenceManager
+class FConcertClientPresenceManager
 {
-	/** Get location update frequency */
-	double GetLocationUpdateFrequency();
-}
+public:
+	static double GetLocationUpdateFrequency();
+};
 
 #endif
diff --git a/Engine/Plugins/Developer/Concert/ConcertSync/ConcertSyncClient/Source/ConcertSyncClient/Private/ConcertClientTransactionBridge.cpp b/Engine/Plugins/Developer/Concert/ConcertSync/ConcertSyncClient/Source/ConcertSyncClient/Private/ConcertClientTransactionBridge.cpp
--- a/Engine/Plugins/Developer/Concert/ConcertSync/ConcertSyncClient/Source/ConcertSyncClient/Private/ConcertClientTransactionBridge.cpp
+++ b/Engine/Plugins/Developer/Concert/ConcertSync/ConcertSyncClient/Source/ConcertSyncClient/Private/ConcertClientTransactionBridge.cpp
@@ -4,14 +4,16 @@
 #include "ConcertLogGlobal.h"
 #include "ConcertSyncSettings.h"
 #include "ConcertSyncClientUtil.h"
-
 #include "ConcertTransactionEvents.h"
 #include "IConcertClientTransactionBridge.h"
-#include "TransactionCommon.h"
+
+#include "GameFramework/Actor.h"
+#include "HAL/IConsoleManager.h"
 #include "Misc/ITransactionObjectAnnotation.h"
 #include "Misc/PackageName.h"
 #include "Misc/CoreDelegates.h"
-#include "HAL/IConsoleManager.h"
+#include "TransactionCommon.h"
+#include "UObject/Package.h"
 
 #if WITH_EDITOR
 	#include "Editor.h"
@@ -59,6 +61,7 @@ bool RunTransactionFilters(const TArray<FTransactionClassFilter>& InFilters, UOb
 	return bMatchFilter;
 }
 
+#if WITH_EDITOR
 template <typename TUObjectType>
 FTypedElementHandle AcquireTypedElementHandle(const TUObjectType* InObject)
 {
@@ -92,6 +95,7 @@ void DeselectElements(UTypedElementSelectionSet* InSelectionSet, const TUObjectA
 	}
 	InSelectionSet->DeselectElements(HandlesToRemoveFromSelection, SelectionOptions);
 }
+#endif
 
 void DeselectActorsAndActorComponents(const TArray<AActor*>& DeletedActors, const TArray<UActorComponent*>& DeletedActorComponents)
 {
@@ -298,7 +302,9 @@ void ProcessTransactionEvent(const FConcertTransactionEventBase& InEvent, const
 	// Phase 1
 	// --------------------------------------------------------------------------------------------------------------------
 	bool bObjectsDeleted = false;
+#if WITH_EDITOR
 	TArray<AActor*> ResurrectedActors;
+#endif
 	TArray<ConcertSyncClientUtil::FGetObjectResult, TInlineAllocator<32>> TransactionObjects;
 	{
 		TSet<const UObject*> NewlyCreatedObjects;
@@ -360,7 +366,11 @@ void ProcessTransactionEvent(const FConcertTransactionEventBase& InEvent, const
 					ConcertSyncUtil::ResetObjectPropertiesToArchetypeValues(TransactionObjectRef.Obj, bIncludeEditorOnlyProperties);
 				}
 
-				if (!TransactionObjectRef.NewlyCreated() && !ObjectUpdate.ObjectData.bIsPendingKill && bWasPendingKill)
+#if WITH_EDITOR
+				// We do not need to handle the resurrection case outside the editor as it involves objects being partially active due to the transaction buffer, a feature exclusive to the editor.
+				// Remember: This case would happen if you undo the deletion of an actor... that does not happen in games!
+				const bool bWasResurrected = !TransactionObjectRef.NewlyCreated() && !ObjectUpdate.ObjectData.bIsPendingKill && bWasPendingKill;
+				if (bWasResurrected)
 				{
 					// If we're bringing this actor back to life, then make sure any SCS/UCS components exist in a clean state
 					// We have to do this as some of the actor components may have been GC'd when using "AllowEliminatingReferences(false)" (eg, within the transaction buffer)
@@ -386,6 +396,7 @@ void ProcessTransactionEvent(const FConcertTransactionEventBase& InEvent, const
 						});
 					}
 				}
+#endif
 			}
 		}
 	}
@@ -571,13 +582,13 @@ void ProcessTransactionEvent(const FConcertTransactionEventBase& InEvent, const
 	{
 		EditorTransactionNotification.PostUndo();
 	}
-#endif
-
+	
 	// Ensure that any actors we restored from the dead are added back to the actors array of their owner level
 	for (AActor* Actor : ResurrectedActors)
 	{
 		ConcertSyncClientUtil::AddActorToOwnerLevel(Actor);
 	}
+#endif
 
 	DeselectActorsAndActorComponents(DeletedActorsForSelectionUpdate, DeletedActorComponentForSelectionUpdate);
 
diff --git a/Engine/Plugins/Developer/Concert/ConcertSync/ConcertSyncClient/Source/ConcertSyncClient/Private/ConcertSyncClientUtil.cpp b/Engine/Plugins/Developer/Concert/ConcertSync/ConcertSyncClient/Source/ConcertSyncClient/Private/ConcertSyncClientUtil.cpp
--- a/Engine/Plugins/Developer/Concert/ConcertSync/ConcertSyncClient/Source/ConcertSyncClient/Private/ConcertSyncClientUtil.cpp
+++ b/Engine/Plugins/Developer/Concert/ConcertSync/ConcertSyncClient/Source/ConcertSyncClient/Private/ConcertSyncClientUtil.cpp
@@ -714,11 +714,14 @@ void PurgePackages(TArrayView<const FName> InPackageNames)
 UWorld* GetCurrentWorld()
 {
 	UWorld* CurrentWorld = nullptr;
+#if WITH_EDITOR
 	if (GIsEditor)
 	{
 		CurrentWorld = GEditor->GetEditorWorldContext().World();
 	}
-	else if (UGameEngine* GameEngine = Cast<UGameEngine>(GEngine))
+	else
+#endif
+	if (UGameEngine* GameEngine = Cast<UGameEngine>(GEngine))
 	{
 		CurrentWorld = GameEngine->GetGameWorld();
 	}
@@ -727,6 +730,7 @@ UWorld* GetCurrentWorld()
 
 ULevel* GetExternalPersistentWorld()
 {
+#if WITH_EDITOR
 	UWorld* CurrentWorld = ConcertSyncClientUtil::GetCurrentWorld();
 	if (CurrentWorld)
 	{
@@ -736,6 +740,7 @@ ULevel* GetExternalPersistentWorld()
 			return PersistentLevel;
 		}
 	}
+#endif
 	return nullptr;
 }
 
