Fix ConcertSyncClient depending on ITargetPlatform in non-editor builds

diff --git a/Engine/Plugins/Developer/Concert/ConcertSync/ConcertSyncClient/Source/ConcertSyncClient/ConcertSyncClient.Build.cs b/Engine/Plugins/Developer/Concert/ConcertSync/ConcertSyncClient/Source/ConcertSyncClient/ConcertSyncClient.Build.cs
--- a/Engine/Plugins/Developer/Concert/ConcertSync/ConcertSyncClient/Source/ConcertSyncClient/ConcertSyncClient.Build.cs
+++ b/Engine/Plugins/Developer/Concert/ConcertSync/ConcertSyncClient/Source/ConcertSyncClient/ConcertSyncClient.Build.cs
@@ -21,14 +21,13 @@ namespace UnrealBuildTool.Rules
 			PrivateDependencyModuleNames.AddRange(
 				new string[]
 				{
-					"ConcertTransport",
 					"AssetRegistry",
+					"ConcertTransport",
 					"HeadMountedDisplay",
 					"InputCore",
-					"MovieScene",
 					"LevelSequence",
+					"MovieScene",
 					"RenderCore",
-					"TargetPlatform",
 					"TimeManagement",
 					"SlateCore",
 					"Slate",
@@ -55,16 +54,16 @@ namespace UnrealBuildTool.Rules
 
 				PrivateDependencyModuleNames.AddRange(
 					new string[] {
+						"EditorFramework",
+						"EditorStyle",
 						"EngineSettings",
+						"LevelEditor",
 						"Sequencer",
-						"EditorFramework",
 						"UnrealEd",
+						"TypedElementRuntime",
+						"TypedElementFramework",
 						"ViewportInteraction",
-						"LevelEditor",
 						"VREditor",
-						"EditorStyle",
-						"TypedElementRuntime",
-						"TypedElementFramework"
 					}
 				);
 			}
diff --git a/Engine/Plugins/Developer/Concert/ConcertSync/ConcertSyncClient/Source/ConcertSyncClient/Private/ConcertClientTransactionManager.cpp b/Engine/Plugins/Developer/Concert/ConcertSync/ConcertSyncClient/Source/ConcertSyncClient/Private/ConcertClientTransactionManager.cpp
--- a/Engine/Plugins/Developer/Concert/ConcertSync/ConcertSyncClient/Source/ConcertSyncClient/Private/ConcertClientTransactionManager.cpp
+++ b/Engine/Plugins/Developer/Concert/ConcertSync/ConcertSyncClient/Source/ConcertSyncClient/Private/ConcertClientTransactionManager.cpp
@@ -1,31 +1,29 @@
 // Copyright Epic Games, Inc. All Rights Reserved.
 
 #include "ConcertClientTransactionManager.h"
-#include "Algo/AnyOf.h"
-#include "Components/SceneComponent.h"
+
 #include "ConcertSyncSessionTypes.h"
 #include "ConcertTransactionEvents.h"
-#include "GameFramework/Actor.h"
 #include "IConcertSession.h"
 #include "ConcertSyncClientLiveSession.h"
 #include "ConcertSyncSessionDatabase.h"
 #include "ConcertLogGlobal.h"
 #include "ConcertSyncSettings.h"
-#include "ConcertSyncArchives.h"
 #include "ConcertSyncClientUtil.h"
 #include "IConcertSyncClient.h"
-#include "Internationalization/Internationalization.h"
 #include "Scratchpad/ConcertScratchpad.h"
 
-#include "Interfaces/ITargetPlatform.h"
-#include "Interfaces/ITargetPlatformManagerModule.h"
-#include "AssetRegistry/AssetRegistryModule.h"
-
+#include "Algo/AnyOf.h"
+#include "Components/SceneComponent.h"
 #include "Engine/World.h"
 #include "Engine/Level.h"
-#include "RenderingThread.h"
+#include "GameFramework/Actor.h"
+#include "Internationalization/Internationalization.h"
 #include "Misc/ScopedSlowTask.h"
+#include "RenderingThread.h"
 #include "UObject/ObjectMacros.h"
+#include "UObject/Package.h"
+#include "UObject/UObjectHash.h"
 
 #if WITH_EDITOR
 	#include "Editor.h"
