Fix Mac/iOS non-unity
#horde 347422
#rnx

diff --git a/Engine/Plugins/Developer/Concert/ConcertSync/ConcertSyncClient/Source/ConcertSyncClient/Private/ConcertClientPackageManager.cpp b/Engine/Plugins/Developer/Concert/ConcertSync/ConcertSyncClient/Source/ConcertSyncClient/Private/ConcertClientPackageManager.cpp
--- a/Engine/Plugins/Developer/Concert/ConcertSync/ConcertSyncClient/Source/ConcertSyncClient/Private/ConcertClientPackageManager.cpp
+++ b/Engine/Plugins/Developer/Concert/ConcertSync/ConcertSyncClient/Source/ConcertSyncClient/Private/ConcertClientPackageManager.cpp
@@ -15,6 +15,8 @@
 #include "ConcertSyncSettings.h"
 #include "ConcertUtil.h"
 
+#include "Algo/Find.h"
+#include "Engine/Level.h"
 #include "Engine/World.h"
 #include "Engine/Engine.h"
 #include "UObject/Package.h"
