More concert non-unity fixes
#rnx
#horde 347422

diff --git a/Engine/Plugins/Developer/Concert/ConcertSync/ConcertSyncClient/Source/ConcertSyncClient/Private/ConcertClientPresenceActor.cpp b/Engine/Plugins/Developer/Concert/ConcertSync/ConcertSyncClient/Source/ConcertSyncClient/Private/ConcertClientPresenceActor.cpp
--- a/Engine/Plugins/Developer/Concert/ConcertSync/ConcertSyncClient/Source/ConcertSyncClient/Private/ConcertClientPresenceActor.cpp
+++ b/Engine/Plugins/Developer/Concert/ConcertSync/ConcertSyncClient/Source/ConcertSyncClient/Private/ConcertClientPresenceActor.cpp
@@ -3,6 +3,7 @@
 #include "ConcertClientPresenceActor.h"
 #include "UObject/ConstructorHelpers.h"
 #include "EngineGlobals.h"
+#include "UObject/Package.h"
 #include "UObject/StructOnScope.h"
 #include "Engine/Scene.h"
 #include "ConcertPresenceEvents.h"
diff --git a/Engine/Plugins/Developer/Concert/ConcertSync/ConcertSyncClient/Source/ConcertSyncClient/Private/ConcertClientWorkspace.cpp b/Engine/Plugins/Developer/Concert/ConcertSync/ConcertSyncClient/Source/ConcertSyncClient/Private/ConcertClientWorkspace.cpp
--- a/Engine/Plugins/Developer/Concert/ConcertSync/ConcertSyncClient/Source/ConcertSyncClient/Private/ConcertClientWorkspace.cpp
+++ b/Engine/Plugins/Developer/Concert/ConcertSync/ConcertSyncClient/Source/ConcertSyncClient/Private/ConcertClientWorkspace.cpp
@@ -34,6 +34,7 @@
 #include "Misc/CoreDelegates.h"
 #include "Misc/ScopedSlowTask.h"
 #include "Misc/FeedbackContext.h"
+#include "Serialization/MemoryReader.h"
 #include "StructDeserializer.h"
 #include "StructSerializer.h"
 #include "Backends/JsonStructDeserializerBackend.h"
