Fix not replicating transforms correctly

diff --git a/Engine/Plugins/Developer/Concert/ConcertSync/ConcertSyncClient/Source/ConcertSyncClient/Private/Replication/Processing/ObjectReplicationApplierProcessor.cpp b/Engine/Plugins/Developer/Concert/ConcertSync/ConcertSyncClient/Source/ConcertSyncClient/Private/Replication/Processing/ObjectReplicationApplierProcessor.cpp
--- a/Engine/Plugins/Developer/Concert/ConcertSync/ConcertSyncClient/Source/ConcertSyncClient/Private/Replication/Processing/ObjectReplicationApplierProcessor.cpp
+++ b/Engine/Plugins/Developer/Concert/ConcertSync/ConcertSyncClient/Source/ConcertSyncClient/Private/Replication/Processing/ObjectReplicationApplierProcessor.cpp
@@ -33,13 +33,19 @@ namespace UE::ConcertSyncClient::Replication
 		GetDataSource().ExtractReplicationDataForObject(Args.ObjectInfo, [this, Object, &bAppliedData](const FConcertSessionSerializedPayload& Payload)
 		{
 			bAppliedData = true;
-			ReplicationFormat->ApplyReplicationEvent(*Object, Payload);
 
-			// TODO: This is very hacky and leaves performance on the table... this is in case ApplyReplicationEvent updates the transform
+			// TODO DP UE-193659: This is very hacky and leaves performance on the table... this is in case ApplyReplicationEvent updates the transform
 			if (USceneComponent* SceneComponent = Cast<USceneComponent>(Object))
 			{
+				FTransform TransformBefore = SceneComponent->GetComponentTransform();
+				ReplicationFormat->ApplyReplicationEvent(*Object, Payload);
+				SceneComponent->SetWorldTransform(TransformBefore);
 				SceneComponent->UpdateComponentToWorld();
 			}
+			else
+			{
+				ReplicationFormat->ApplyReplicationEvent(*Object, Payload);
+			}
 		});
 		// This should not happen. If it does, we're wasting  network bandwidth.
 		UE_CLOG(!bAppliedData, LogConcert, Warning, TEXT("Replication: Server sent data that could not be applied (likely it was empty) for object %s from stream %s"), *Args.ObjectInfo.Object.ToString(), *Args.ObjectInfo.StreamId.ToString());
