Fixed Concert calculating an invalid sandbox path for a project outside of the normal location

Projects under the UE Sandbox directory could end up with paths like "../../Content" which would return "../.." as the folder name, resulting in a mapping that escaped the Concert sandbox directory.

During shutdown this bad path would result in deleting the entire Concert session directory rather than just the paths within the Concert sandbox.

Given the comment in GetContentFolderName said it was just trying to get the content folder name from "/Bla/Content", and given the mount root already follows the form "/Bla", it was simplest to just append the mount root directly to the Concert sandbox root.

#rnx

diff --git a/Engine/Plugins/Developer/Concert/ConcertSync/ConcertSyncClient/Source/ConcertSyncClient/Private/ConcertSandboxPlatformFile.cpp b/Engine/Plugins/Developer/Concert/ConcertSync/ConcertSyncClient/Source/ConcertSyncClient/Private/ConcertSandboxPlatformFile.cpp
--- a/Engine/Plugins/Developer/Concert/ConcertSync/ConcertSyncClient/Source/ConcertSyncClient/Private/ConcertSandboxPlatformFile.cpp
+++ b/Engine/Plugins/Developer/Concert/ConcertSync/ConcertSyncClient/Source/ConcertSyncClient/Private/ConcertSandboxPlatformFile.cpp
@@ -23,24 +23,6 @@
 namespace ConcertSandboxPlatformFileUtil
 {
 
-FString GetContentFolderName(const FString& InContentPath)
-{
-	check(InContentPath.Len() > 0);
-
-	// Clean any trailing slash from the content path
-	FString ContentFolderName = InContentPath;
-	if (ContentFolderName[ContentFolderName.Len() - 1] == TEXT('/') || ContentFolderName[ContentFolderName.Len() - 1] == TEXT('\\'))
-	{
-		ContentFolderName.RemoveAt(ContentFolderName.Len() - 1, 1, /*bAllowShrinking*/false);
-	}
-
-	// Content paths are always in the form /Bla/Content, so we need to use GetBaseFilename after calling GetPath to get the 'Bla' part for the sandbox path
-	ContentFolderName = FPaths::GetPath(MoveTemp(ContentFolderName));
-	ContentFolderName = FPaths::GetCleanFilename(MoveTemp(ContentFolderName));
-
-	return ContentFolderName;
-}
-
 bool FlushPackageFile(const FString& InFilename, FName* OutPackageName = nullptr, bool bForceLoad = true)
 {
 	FString PackageName;
@@ -119,7 +101,7 @@ bool FConcertSandboxPlatformFile::Initialize(IPlatformFile* Inner, const TCHAR*
 	FPackageName::QueryRootContentPaths(RootPaths);
 	for (const FString& RootPath : RootPaths)
 	{
-		RegisterContentMountPath(FPackageName::LongPackageNameToFilename(RootPath));
+		RegisterContentMountPath(RootPath, FPackageName::LongPackageNameToFilename(RootPath));
 	}
 
 	// Watch for new content mount paths
@@ -1056,18 +1038,18 @@ TArray<FString> FConcertSandboxPlatformFile::GatherSandboxChangedFilenames() con
 
 void FConcertSandboxPlatformFile::OnContentPathMounted(const FString& InAssetPath, const FString& InFilesystemPath)
 {
-	RegisterContentMountPath(InFilesystemPath);
+	RegisterContentMountPath(InAssetPath, InFilesystemPath);
 }
 
 void FConcertSandboxPlatformFile::OnContentPathDismounted(const FString& InAssetPath, const FString& InFilesystemPath)
 {
-	UnregisterContentMountPath(InFilesystemPath);
+	UnregisterContentMountPath(InAssetPath, InFilesystemPath);
 }
 
-void FConcertSandboxPlatformFile::RegisterContentMountPath(const FString& InContentPath)
+void FConcertSandboxPlatformFile::RegisterContentMountPath(const FString& InAssetPath, const FString& InFilesystemPath)
 {
-	FString AbsoluteSandboxPath = FPaths::ConvertRelativePathToFull(SandboxRootPath / ConcertSandboxPlatformFileUtil::GetContentFolderName(InContentPath)) / TEXT("");
-	FString AbsoluteNonSandboxPath = FPaths::ConvertRelativePathToFull(InContentPath) / TEXT("");
+	FString AbsoluteSandboxPath = FPaths::ConvertRelativePathToFull(SandboxRootPath / InAssetPath) / TEXT("");
+	FString AbsoluteNonSandboxPath = FPaths::ConvertRelativePathToFull(InFilesystemPath) / TEXT("");
 	LowerLevel->CreateDirectory(*AbsoluteSandboxPath);
 	{
 		FScopeLock SandboxMountPointsLock(&SandboxMountPointsCS);
@@ -1082,9 +1064,9 @@ void FConcertSandboxPlatformFile::RegisterContentMountPath(const FString& InCont
 	}
 }
 
-void FConcertSandboxPlatformFile::UnregisterContentMountPath(const FString& InContentPath)
+void FConcertSandboxPlatformFile::UnregisterContentMountPath(const FString& InAssetPath, const FString& InFilesystemPath)
 {
-	const FString AbsoluteSandboxPath = FPaths::ConvertRelativePathToFull(SandboxRootPath / ConcertSandboxPlatformFileUtil::GetContentFolderName(InContentPath)) / TEXT("");
+	const FString AbsoluteSandboxPath = FPaths::ConvertRelativePathToFull(SandboxRootPath / InAssetPath) / TEXT("");
 	{
 		FScopeLock SandboxMountPointsLock(&SandboxMountPointsCS);
 		SandboxMountPoints.RemoveAll([&AbsoluteSandboxPath](FSandboxMountPoint& InSandboxMountPoint) -> bool
diff --git a/Engine/Plugins/Developer/Concert/ConcertSync/ConcertSyncClient/Source/ConcertSyncClient/Private/ConcertSandboxPlatformFile.h b/Engine/Plugins/Developer/Concert/ConcertSync/ConcertSyncClient/Source/ConcertSyncClient/Private/ConcertSandboxPlatformFile.h
--- a/Engine/Plugins/Developer/Concert/ConcertSync/ConcertSyncClient/Source/ConcertSyncClient/Private/ConcertSandboxPlatformFile.h
+++ b/Engine/Plugins/Developer/Concert/ConcertSync/ConcertSyncClient/Source/ConcertSyncClient/Private/ConcertSandboxPlatformFile.h
@@ -228,10 +228,10 @@ private:
 	void OnContentPathDismounted(const FString& InAssetPath, const FString& InFilesystemPath);
 
 	/** Register a mount path from a source content path */
-	void RegisterContentMountPath(const FString& InContentPath);
+	void RegisterContentMountPath(const FString& InAssetPath, const FString& InFilesystemPath);
 
 	/** Unregister a mount path from a source content path */
-	void UnregisterContentMountPath(const FString& InContentPath);
+	void UnregisterContentMountPath(const FString& InAssetPath, const FString& InFilesystemPath);
 
 	/** Resolve the given path to its sandbox path (if any) */
 	FConcertSandboxPlatformFilePath ToSandboxPath(FString InFilename, const bool bEvenIfDisabled = false) const;
