Check that we have serialized data prior to reading CreationMethod data.

diff --git a/Engine/Plugins/Developer/Concert/ConcertSync/ConcertSyncClient/Source/ConcertSyncClient/Private/ConcertClientTransactionBridge.cpp b/Engine/Plugins/Developer/Concert/ConcertSync/ConcertSyncClient/Source/ConcertSyncClient/Private/ConcertClientTransactionBridge.cpp
--- a/Engine/Plugins/Developer/Concert/ConcertSync/ConcertSyncClient/Source/ConcertSyncClient/Private/ConcertClientTransactionBridge.cpp
+++ b/Engine/Plugins/Developer/Concert/ConcertSync/ConcertSyncClient/Source/ConcertSyncClient/Private/ConcertClientTransactionBridge.cpp
@@ -339,7 +339,8 @@ void ProcessTransactionEvent(const FConcertTransactionEventBase& InEvent, const
 					if (UActorComponent* Component = Cast<UActorComponent>(TransactionObjectRef.Obj))
 					{
 						static const FName NAME_CreationMethod = "CreationMethod";
-						if (const FConcertSerializedPropertyData* CreationMethodPropertyData = ObjectUpdate.PropertyDatas.FindByPredicate([](const FConcertSerializedPropertyData& PropertyData) { return PropertyData.PropertyName == NAME_CreationMethod; }))
+						const FConcertSerializedPropertyData* CreationMethodPropertyData = ObjectUpdate.PropertyDatas.FindByPredicate([](const FConcertSerializedPropertyData& PropertyData) { return PropertyData.PropertyName == NAME_CreationMethod; });
+						if (CreationMethodPropertyData && !CreationMethodPropertyData->SerializedData.IsEmpty())
 						{
 							if (FProperty* CreationMethodProperty = FindFProperty<FProperty>(TransactionObjectRef.Obj->GetClass(), CreationMethodPropertyData->PropertyName))
 							{
