Do not let temporary packages leak into session.  They are  transient and will be saved later under a proper name later.

diff --git a/Engine/Plugins/Developer/Concert/ConcertSync/ConcertSyncClient/Source/ConcertSyncClient/Private/ConcertClientPackageManager.cpp b/Engine/Plugins/Developer/Concert/ConcertSync/ConcertSyncClient/Source/ConcertSyncClient/Private/ConcertClientPackageManager.cpp
--- a/Engine/Plugins/Developer/Concert/ConcertSync/ConcertSyncClient/Source/ConcertSyncClient/Private/ConcertClientPackageManager.cpp
+++ b/Engine/Plugins/Developer/Concert/ConcertSync/ConcertSyncClient/Source/ConcertSyncClient/Private/ConcertClientPackageManager.cpp
@@ -484,13 +484,14 @@ void FConcertClientPackageManager::AddPendingReloadForNewExternalMaps(const FCon
 void FConcertClientPackageManager::HandleLocalPackageEvent(const FConcertPackageInfo& PackageInfo, const FString& PackagePathname)
 {
 	// Ignore unwanted saves
+	const bool bIsValidPackageName = FPackageName::IsValidLongPackageName(PackageInfo.PackageName.ToString());
+	if (!bIsValidPackageName)
+	{
+		return;
+	}
 	if (PackageInfo.PackageUpdateType == EConcertPackageUpdateType::Saved)
 	{
-		if (!FPackageName::IsValidLongPackageName(PackageInfo.PackageName.ToString())) // Auto-Save might save the template in /Temp/... which is an invalid long package name.
-		{
-			return;
-		}
-		else if (PackageInfo.bAutoSave && !EnumHasAnyFlags(LiveSession->GetSessionFlags(), EConcertSyncSessionFlags::ShouldSendPackageAutoSaves))
+		if (PackageInfo.bAutoSave && !EnumHasAnyFlags(LiveSession->GetSessionFlags(), EConcertSyncSessionFlags::ShouldSendPackageAutoSaves))
 		{
 			return;
 		}
