When invoking the OpenAssetInEditor ensure we are referring to the level sequence object and not the pending take.

diff --git a/Engine/Plugins/Developer/Concert/ConcertSync/ConcertSyncClient/Source/ConcertSyncClient/Private/ConcertClientSequencerManager.cpp b/Engine/Plugins/Developer/Concert/ConcertSync/ConcertSyncClient/Source/ConcertSyncClient/Private/ConcertClientSequencerManager.cpp
--- a/Engine/Plugins/Developer/Concert/ConcertSync/ConcertSyncClient/Source/ConcertSyncClient/Private/ConcertClientSequencerManager.cpp
+++ b/Engine/Plugins/Developer/Concert/ConcertSync/ConcertSyncClient/Source/ConcertSyncClient/Private/ConcertClientSequencerManager.cpp
@@ -779,7 +779,10 @@ void FConcertClientSequencerManager::ApplyOpenEvent(const FConcertSequencerOpenE
 		#if WITH_EDITOR
 		if (GIsEditor)
 		{
-			GEditor->GetEditorSubsystem<UAssetEditorSubsystem>()->OpenEditorForAsset(SequenceObjectPath);
+			// We must specify the pending level sequence and not the pending take.  If we specify the pending take after we serialized into the existing
+			// asset the open asset will cause take recorder tab to assume that we want to edit the pending take where we really just want to view
+			// the level sequence.
+			GEditor->GetEditorSubsystem<UAssetEditorSubsystem>()->OpenEditorForAsset(PendingLevelSequence);
 		}
 		#endif
 		PendingLevelSequence->GetPackage()->SetDirtyFlag(false);
