Ignore objects nested within transient objects

This addresses an edge case where an object within a transient outer has its inherited RF_Transient flag cleared, as that object will still ultimately be transient when saved due to its outer being transient.

#rnx

diff --git a/Engine/Plugins/Developer/Concert/ConcertSync/ConcertSyncClient/Source/ConcertSyncClient/Private/ConcertClientTransactionBridge.cpp b/Engine/Plugins/Developer/Concert/ConcertSync/ConcertSyncClient/Source/ConcertSyncClient/Private/ConcertClientTransactionBridge.cpp
--- a/Engine/Plugins/Developer/Concert/ConcertSync/ConcertSyncClient/Source/ConcertSyncClient/Private/ConcertClientTransactionBridge.cpp
+++ b/Engine/Plugins/Developer/Concert/ConcertSync/ConcertSyncClient/Source/ConcertSyncClient/Private/ConcertClientTransactionBridge.cpp
@@ -134,13 +134,26 @@ ETransactionFilterResult ApplyCustomFilter(const TMap<FName, FTransactionFilterD
 
 ETransactionFilterResult ApplyTransactionFilters(const TMap<FName, FTransactionFilterDelegate>& CustomFilters, UObject* InObject, UPackage* InChangedPackage)
 {
+	// An object is persistent if neither it nor any of its outers are transient
+	auto IsObjectPersistent = [](const UObject* Obj)
+	{
+		for (; Obj; Obj = Obj->GetOuter())
+		{
+			if (Obj->HasAnyFlags(RF_Transient))
+			{
+				return false;
+			}
+		}
+		return true;
+	};
+
 	ETransactionFilterResult FilterResult = ConcertClientTransactionBridgeUtil::ApplyCustomFilter(CustomFilters, InObject, InChangedPackage);
 	if (FilterResult != ETransactionFilterResult::UseDefault)
 	{
 		return FilterResult;
 	}
 	// Ignore transient packages and objects, compiled in package are not considered Multi-user content.
-	if (!InChangedPackage || InChangedPackage == GetTransientPackage() || InChangedPackage->HasAnyFlags(RF_Transient) || InChangedPackage->HasAnyPackageFlags(PKG_CompiledIn) || InObject->HasAnyFlags(RF_Transient))
+	if (!InChangedPackage || InChangedPackage == GetTransientPackage() || InChangedPackage->HasAnyPackageFlags(PKG_CompiledIn) || !IsObjectPersistent(InObject))
 	{
 		return ETransactionFilterResult::ExcludeObject;
 	}
