Make ConcertTransactionBridge aware of LevelInstance changes in editor so that the can be properly replayed on -game instances. Leverages the experimental EditorPaths helper.

diff --git a/Engine/Plugins/Developer/Concert/ConcertSync/ConcertSyncClient/Source/ConcertSyncClient/Private/ConcertClientTransactionBridge.cpp b/Engine/Plugins/Developer/Concert/ConcertSync/ConcertSyncClient/Source/ConcertSyncClient/Private/ConcertClientTransactionBridge.cpp
--- a/Engine/Plugins/Developer/Concert/ConcertSync/ConcertSyncClient/Source/ConcertSyncClient/Private/ConcertClientTransactionBridge.cpp
+++ b/Engine/Plugins/Developer/Concert/ConcertSync/ConcertSyncClient/Source/ConcertSyncClient/Private/ConcertClientTransactionBridge.cpp
@@ -14,6 +14,8 @@
 #include "Misc/CoreDelegates.h"
 #include "TransactionCommon.h"
 #include "UObject/Package.h"
+#include "LevelInstance/LevelInstanceActor.h"
+#include "LevelInstance/LevelInstanceComponent.h"
 
 #if WITH_EDITOR
 	#include "Editor.h"
@@ -284,6 +286,19 @@ struct FEditorTransactionNotification
 };
 #endif
 
+FConcertExportedObject MakeLevelInstanceExportedObject(const FConcertObjectId& LevelInstanceId, const FConcertExportedObject& SourceObject)
+{
+	FConcertExportedObject OutExportedObject(
+		LevelInstanceId,                // FConcertObjectId
+		SourceObject.ObjectPathDepth,   // int32 ObjectPathDepth
+		{},                             // No level instance objects.
+		SourceObject.ObjectData,
+		SourceObject.PropertyDatas,
+		{}                              // We don't support annotation data with Level Instances. All changes should be property deltas.
+		);
+
+	return OutExportedObject;
+}
 
 void ProcessTransactionEvent(const FConcertTransactionEventBase& InEvent, const FConcertSessionVersionInfo* InVersionInfo, const TArray<FName>& InPackagesToProcess, const FConcertLocalIdentifierTable* InLocalIdentifierTablePtr, const bool bIsSnapshot, const FConcertSyncWorldRemapper& WorldRemapper, const bool bIncludeEditorOnlyProperties)
 {
@@ -297,11 +312,46 @@ void ProcessTransactionEvent(const FConcertTransactionEventBase& InEvent, const
 	// --------------------------------------------------------------------------------------------------------------------
 	// Phase 0
 	// --------------------------------------------------------------------------------------------------------------------
+	int32 LevelInstanceObjects = 0;
+
+#if WITH_EDITOR
+	TArray<FConcertExportedObject> AddedLevelInstanceObjects;
+	// For -game instances, we need to gather all affected Level Instance Actors and apply the delta change to those
+	// objects.
+	if (!GIsEditor)
+	{
+		for (const FConcertExportedObject& ExportedObject : InEvent.ExportedObjects)
+		{
+			LevelInstanceObjects += ExportedObject.LevelInstanceObjects.Num();
+		}
+		AddedLevelInstanceObjects.Reserve(LevelInstanceObjects);
+		if (LevelInstanceObjects > 0 && ConcertSyncClientUtil::GetExternalPersistentWorld())
+		{
+			// This is a WP or OFPA level and we cannot apply transactions in those worlds due to limitations
+			// of using LevelInstance.ForceEditorWorldMode=1 and cell generation.  Future support for OFPA / WP in Multi-user
+			// with LevelInstances is TBD. Non WP and -game nodes can work with Level Instances.
+			//
+			return;
+		}
+	}
+#endif
+
 	TArray<const FConcertExportedObject*, TInlineAllocator<32>> SortedExportedObjects;
 	{
-		SortedExportedObjects.Reserve(InEvent.ExportedObjects.Num());
+		SortedExportedObjects.Reserve(InEvent.ExportedObjects.Num() + LevelInstanceObjects);
 		for (const FConcertExportedObject& ExportedObject : InEvent.ExportedObjects)
 		{
+#if WITH_EDITOR
+			// If the exported object is an actor in a level instance and we are not in -game then skip.
+			if (!GIsEditor)
+			{
+				for (const FConcertObjectId& Id : ExportedObject.LevelInstanceObjects)
+				{
+					AddedLevelInstanceObjects.Emplace(MakeLevelInstanceExportedObject(Id, ExportedObject));
+					SortedExportedObjects.Add(&(AddedLevelInstanceObjects.Last()));
+				}
+			}
+#endif
 			SortedExportedObjects.Add(&ExportedObject);
 		}
 
@@ -783,37 +833,14 @@ void FConcertClientTransactionBridge::HandleTransactionStateChanged(const FTrans
 	}
 }
 
-void FConcertClientTransactionBridge::HandleObjectTransacted(UObject* InObject, const FTransactionObjectEvent& InTransactionEvent)
-{
-	LLM_SCOPE_BYTAG(Concert_ConcertClientTransactionBridge);
-	SCOPED_CONCERT_TRACE(FConcertClientTransactionBridge_HandleObjectTransacted);
-
-	if (bIgnoreLocalTransactions)
-	{
-		return;
-	}
-
-	if (!bIncludeAnnotationObjectChanges && InTransactionEvent.GetObjectChangeCreatedBy() == ETransactionObjectChangeCreatedBy::TransactionAnnotation)
-	{
-		return;
-	}
-
-	if (!bIncludeEditorOnlyProperties && InObject->IsEditorOnly())
-	{
-		return;
-	}
-
-	UPackage* ChangedPackage = InObject->GetOutermost();
-	ETransactionFilterResult FilterResult = ConcertClientTransactionBridgeUtil::ApplyTransactionFilters(TransactionFilters, InObject, ChangedPackage);
-	FOngoingTransaction* TrackedTransaction = OngoingTransactions.Find(InTransactionEvent.GetOperationId());
-
-	// TODO: This needs to send both editor-only and non-editor-only payload data to the server, which will forward only the correct part to cooked and non-cooked clients
-
 
+namespace ConcertClientTransactionBridgeUtil
+{
+void LogTransactionEventDetails(UObject* InObject, const FTransactionObjectEvent& InTransactionEvent, ETransactionFilterResult FilterResult, bool bIsTracked)
+{
+	const TCHAR* ObjectEventString = TEXT("");
+	switch (InTransactionEvent.GetEventType())
 	{
-		const TCHAR* ObjectEventString = TEXT("");
-		switch (InTransactionEvent.GetEventType())
-		{
 #define ENUM_TO_STRING(ENUM)						\
 		case ETransactionObjectEventType::ENUM:		\
 			ObjectEventString = TEXT(#ENUM);		\
@@ -824,40 +851,71 @@ void FConcertClientTransactionBridge::HandleObjectTransacted(UObject* InObject,
 #undef ENUM_TO_STRING
 		default:
 			break;
-		}
+	}
 
-		UE_LOG(LogConcert, VeryVerbose,
-			TEXT("%s Transaction %s (%s, %s):%s %s:%s (%s property changes, %s object changes)"), 
-			(TrackedTransaction != nullptr ? TEXT("Tracked") : TEXT("Untracked")),
-			*InTransactionEvent.GetTransactionId().ToString(),
-			*InTransactionEvent.GetOperationId().ToString(),
-			ObjectEventString,
-			(FilterResult == ETransactionFilterResult::ExcludeObject ? TEXT(" FILTERED OBJECT: ") : TEXT("")),
-			*InObject->GetClass()->GetName(),
-			*InObject->GetPathName(), 
-			(InTransactionEvent.HasPropertyChanges() ? TEXT("has") : TEXT("no")), 
-			(InTransactionEvent.HasNonPropertyChanges() ? TEXT("has") : TEXT("no"))
+	UE_LOG(LogConcert, VeryVerbose,
+		   TEXT("%s Transaction %s (%s, %s):%s %s:%s (%s property changes, %s object changes)"),
+		   (bIsTracked ? TEXT("Tracked") : TEXT("Untracked")),
+		   *InTransactionEvent.GetTransactionId().ToString(),
+		   *InTransactionEvent.GetOperationId().ToString(),
+		   ObjectEventString,
+		   (FilterResult == ETransactionFilterResult::ExcludeObject ? TEXT(" FILTERED OBJECT: ") : TEXT("")),
+		   *InObject->GetClass()->GetName(),
+		   *InObject->GetPathName(),
+		   (InTransactionEvent.HasPropertyChanges() ? TEXT("has") : TEXT("no")),
+		   (InTransactionEvent.HasNonPropertyChanges() ? TEXT("has") : TEXT("no"))
 			);
+}
+
+UClass* GetModifiedClass(UObject* InObject)
+{
+	if (Cast<ALevelInstance>(InObject))
+	{
+		return AActor::StaticClass();
+	}
+	if (Cast<ULevelInstanceComponent>(InObject) || Cast<USceneComponent>(InObject))
+	{
+		return USceneComponent::StaticClass();
 	}
+	return nullptr;
+}
 
-	if (TrackedTransaction == nullptr)
+ALevelInstance* GetLevelInstance(UObject* InObject)
+{
+	ALevelInstance* LevelInstance = Cast<ALevelInstance>(InObject);
+	if (!LevelInstance && InObject)
 	{
-		return;
+		UObject* Owner = InObject->GetOuter();
+		LevelInstance = Cast<ALevelInstance>(Owner);
 	}
 
-	const FConcertObjectId ObjectId = FConcertObjectId(InTransactionEvent.GetOriginalObjectId(), InObject->GetFlags());
-	FOngoingTransaction& OngoingTransaction = *TrackedTransaction;
+	return LevelInstance;
+}
 
-	// If the object is excluded or exclude the whole transaction add it to the excluded list
-	if (FilterResult != ETransactionFilterResult::IncludeObject)
+void ApplyForAllRelevantObjects(UObject* InObject, TFunctionRef<void(UObject*)> InFunc, UClass* InClass = nullptr)
+{
+	if (ALevelInstance* LevelInstance = GetLevelInstance(InObject))
 	{
-		OngoingTransaction.CommonData.bIsExcluded |= FilterResult == ETransactionFilterResult::ExcludeTransaction;
-		OngoingTransaction.CommonData.ExcludedObjectUpdates.Add(ObjectId);
-		return;
+		UClass* TheClass = InClass ? InClass : GetModifiedClass(InObject);
+
+		// We should always have a class at this point.
+		check(TheClass);
+		TSet<AActor*> Actors;
+		LevelInstance->EditorGetUnderlyingActors(Actors);
+		for (AActor* Actor : Actors)
+		{
+			UObject* ObjectToApply = Actor->FindComponentByClass(TheClass);
+			if (ObjectToApply)
+			{
+				InFunc(ObjectToApply);
+				ApplyForAllRelevantObjects(Actor, InFunc, TheClass);
+			}
+		}
 	}
+}
 
-	bool bCanCreateObject = true;
-	bool bCanRenameObject = true;
+void GetCanCreateOrRenameObject(UObject* InObject, bool& bOutCanCreateObject, bool& bOutCanRenameObject)
+{
 	if (const UActorComponent* Component = Cast<UActorComponent>(InObject))
 	{
 		bool bIsOwnedByChildActor = false;
@@ -867,97 +925,116 @@ void FConcertClientTransactionBridge::HandleObjectTransacted(UObject* InObject,
 		}
 
 		// Components that are managed by a construction script cannot be created
-		bCanCreateObject = !bIsOwnedByChildActor && !Component->IsCreatedByConstructionScript();
+		bOutCanCreateObject = !bIsOwnedByChildActor && !Component->IsCreatedByConstructionScript();
 
 		// Components that are managed by a native or Blueprint class cannot be renamed, so we only allow "instance" components to be created or renamed
-		bCanRenameObject = !bIsOwnedByChildActor && Component->CreationMethod == EComponentCreationMethod::Instance;
+		bOutCanRenameObject = !bIsOwnedByChildActor && Component->CreationMethod == EComponentCreationMethod::Instance;
 	}
+}
 
-	const FName NewObjectPackageName = (bCanRenameObject && (InTransactionEvent.HasOuterChange() || InTransactionEvent.HasExternalPackageChange())) ? InObject->GetPackage()->GetFName() : FName();
-	const FName NewObjectName = (bCanRenameObject && InTransactionEvent.HasNameChange()) ? InObject->GetFName() : FName();
-	const FName NewObjectOuterPathName = (bCanRenameObject && (InTransactionEvent.HasOuterChange() && InObject->GetOuter())) ? FName(*InObject->GetOuter()->GetPathName()) : FName();
-	const FName NewObjectExternalPackageName = (bCanRenameObject && (InTransactionEvent.HasExternalPackageChange() && InObject->GetExternalPackage())) ? InObject->GetExternalPackage()->GetFName() : FName();
-	const TArray<const FProperty*> ExportedProperties = ConcertSyncClientUtil::GetExportedProperties(InObject->GetClass(), InTransactionEvent.GetChangedProperties(), bIncludeEditorOnlyProperties);
-	TSharedPtr<ITransactionObjectAnnotation> TransactionAnnotation = InTransactionEvent.GetAnnotation();
-	const bool bUseSerializedAnnotationData = !bIncludeAnnotationObjectChanges || !TransactionAnnotation || !TransactionAnnotation->SupportsAdditionalObjectChanges();
-
-	const bool bObjectHasChangesToSend = InTransactionEvent.HasIdOrPendingKillChanges()
-		|| (bIncludeNonPropertyObjectData && InTransactionEvent.HasNonPropertyChanges(/*SerializationOnly*/true))
-		|| ExportedProperties.Num() > 0
-		|| (TransactionAnnotation && bUseSerializedAnnotationData);
-
-	if (!bObjectHasChangesToSend)
-	{
-		// This object has no changes to send
-		return;
-	}
-
-	// Track which packages were changed
-	OngoingTransaction.CommonData.ModifiedPackages.AddUnique(ChangedPackage->GetFName());
-
-	// if there was an outer change, track that outer package
-	if (InTransactionEvent.HasOuterChange())
-	{
-		FName OriginalOuterPackageName = *FPackageName::ObjectPathToPackageName(InTransactionEvent.GetOriginalObjectOuterPathName().ToString());
-		OngoingTransaction.CommonData.ModifiedPackages.AddUnique(OriginalOuterPackageName);
-	}
+struct FTransactedObjectState
+{
+    bool bCanCreateObject = true;
+	bool bCanRenameObject = true;
 
-	// if there was an package change, track that package
-	if (InTransactionEvent.HasExternalPackageChange())
+	using FOngoingTransaction = FConcertClientTransactionBridge::FOngoingTransaction;
+
+	const FTransactionObjectEvent& TransactionEvent;
+	FOngoingTransaction& OngoingTransaction;
+
+	FName NewObjectPackageName;
+	FName NewObjectName;
+	FName NewObjectOuterPathName;
+	FName NewObjectExternalPackageName;
+	TArray<const FProperty*> ExportedProperties;
+	TSharedPtr<ITransactionObjectAnnotation> TransactionAnnotation;
+	bool bUseSerializedAnnotationData;
+
+	bool bIncludeNonPropertyObjectData = false;
+	bool bIncludeEditorOnlyProperties = false;
+
+	FTransactedObjectState() = delete;
+    FTransactedObjectState(UObject* InObject, const FTransactionObjectEvent& InTransactionEvent,
+						   FOngoingTransaction& InOngoingTransaction,
+						   bool bInIncludeNonPropertyObjectData,
+						   bool bInIncludeEditorOnlyProperties,
+						   bool bIncludeAnnotationObjectChanges)
+		: TransactionEvent(InTransactionEvent),
+		  OngoingTransaction(InOngoingTransaction)
+    {
+		bIncludeNonPropertyObjectData =  bInIncludeNonPropertyObjectData;
+		bIncludeEditorOnlyProperties = bInIncludeEditorOnlyProperties;
+
+        ConcertClientTransactionBridgeUtil::GetCanCreateOrRenameObject(InObject, bCanCreateObject, bCanRenameObject);
+
+        if (bCanRenameObject)
+        {
+            if (TransactionEvent.HasOuterChange() || TransactionEvent.HasExternalPackageChange())
+            {
+                NewObjectPackageName = InObject->GetPackage()->GetFName();
+            }
+            if (TransactionEvent.HasNameChange())
+            {
+                NewObjectName = InObject->GetFName();
+            }
+            if (TransactionEvent.HasOuterChange() && InObject->GetOuter())
+            {
+                NewObjectOuterPathName = FName(*InObject->GetOuter()->GetPathName()) ;
+            }
+            if (TransactionEvent.HasExternalPackageChange() && InObject->GetExternalPackage())
+            {
+                NewObjectExternalPackageName = InObject->GetExternalPackage()->GetFName();
+            }
+        }
+
+        ExportedProperties = ConcertSyncClientUtil::GetExportedProperties(InObject->GetClass(), TransactionEvent.GetChangedProperties(), bIncludeEditorOnlyProperties);
+
+        TransactionAnnotation = TransactionEvent.GetAnnotation();
+        bUseSerializedAnnotationData = !bIncludeAnnotationObjectChanges
+            || !TransactionAnnotation
+            || !TransactionAnnotation->SupportsAdditionalObjectChanges();
+    }
+
+	explicit operator bool() const
+    {
+        const bool bObjectHasChangesToSend = TransactionEvent.HasIdOrPendingKillChanges()
+            || (bIncludeNonPropertyObjectData && TransactionEvent.HasNonPropertyChanges(/*SerializationOnly*/true))
+            || ExportedProperties.Num() > 0
+            || (TransactionAnnotation && bUseSerializedAnnotationData);
+        return bObjectHasChangesToSend;
+    }
+
+	void CaptureSnapshot(UObject* InObject, FConcertExportedObject* ObjectUpdatePtr)
 	{
-		FName OriginalPackageName = InTransactionEvent.GetOriginalObjectPackageName().IsNone() ? *FPackageName::ObjectPathToPackageName(InTransactionEvent.GetOriginalObjectOuterPathName().ToString()) : InTransactionEvent.GetOriginalObjectPackageName();
-		OngoingTransaction.CommonData.ModifiedPackages.AddUnique(OriginalPackageName);
-	}
+		if (TransactionAnnotation && bUseSerializedAnnotationData)
+		{
+			ObjectUpdatePtr->SerializedAnnotationData.Reset();
+			FConcertSyncObjectWriter AnnotationWriter(nullptr, InObject, ObjectUpdatePtr->SerializedAnnotationData, bIncludeEditorOnlyProperties, true);
+			TransactionAnnotation->Serialize(AnnotationWriter);
+		}
 
-	// Add this object change to its pending transaction
-	if (InTransactionEvent.GetEventType() == ETransactionObjectEventType::Snapshot)
-	{
-		// Merge the snapshot property changes into pending snapshot list
-		if (OnLocalTransactionSnapshotDelegate.IsBound() && (ExportedProperties.Num() > 0 || TransactionAnnotation.IsValid()))
+		// Find or add an update for each property
+		for (const FProperty* ExportedProperty : ExportedProperties)
 		{
-			// Find or add an entry for this object
-			FConcertExportedObject* ObjectUpdatePtr = OngoingTransaction.SnapshotData.SnapshotObjectUpdates.FindByPredicate([&ObjectId](FConcertExportedObject& ObjectUpdate)
+			FConcertSerializedPropertyData* PropertyDataPtr = ObjectUpdatePtr->PropertyDatas.FindByPredicate([ExportedProperty](FConcertSerializedPropertyData& PropertyData)
 			{
-				return ConcertSyncClientUtil::ObjectIdsMatch(ObjectId, ObjectUpdate.ObjectId);
+				return ExportedProperty->GetFName() == PropertyData.PropertyName;
 			});
-			if (!ObjectUpdatePtr)
-			{
-				ObjectUpdatePtr = &OngoingTransaction.SnapshotData.SnapshotObjectUpdates.AddDefaulted_GetRef();
-				ObjectUpdatePtr->ObjectId = ObjectId;
-				ObjectUpdatePtr->ObjectPathDepth = ConcertSyncClientUtil::GetObjectPathDepth(InObject);
-				ObjectUpdatePtr->ObjectData.bIsPendingKill = !IsValid(InObject);
-			}
-
-			if (TransactionAnnotation && bUseSerializedAnnotationData)
+			if (!PropertyDataPtr)
 			{
-				ObjectUpdatePtr->SerializedAnnotationData.Reset();
-				FConcertSyncObjectWriter AnnotationWriter(nullptr, InObject, ObjectUpdatePtr->SerializedAnnotationData, bIncludeEditorOnlyProperties, true);
-				TransactionAnnotation->Serialize(AnnotationWriter);
+				PropertyDataPtr = &ObjectUpdatePtr->PropertyDatas.AddDefaulted_GetRef();
+				PropertyDataPtr->PropertyName = ExportedProperty->GetFName();
 			}
 
-			// Find or add an update for each property
-			for (const FProperty* ExportedProperty : ExportedProperties)
-			{
-				FConcertSerializedPropertyData* PropertyDataPtr = ObjectUpdatePtr->PropertyDatas.FindByPredicate([ExportedProperty](FConcertSerializedPropertyData& PropertyData)
-				{
-					return ExportedProperty->GetFName() == PropertyData.PropertyName;
-				});
-				if (!PropertyDataPtr)
-				{
-					PropertyDataPtr = &ObjectUpdatePtr->PropertyDatas.AddDefaulted_GetRef();
-					PropertyDataPtr->PropertyName = ExportedProperty->GetFName();
-				}
-
-				PropertyDataPtr->SerializedData.Reset();
-				ConcertSyncClientUtil::SerializeProperty(nullptr, InObject, ExportedProperty, bIncludeEditorOnlyProperties, PropertyDataPtr->SerializedData);
-			}
+			PropertyDataPtr->SerializedData.Reset();
+			ConcertSyncClientUtil::SerializeProperty(nullptr, InObject, ExportedProperty, bIncludeEditorOnlyProperties, PropertyDataPtr->SerializedData);
 		}
 	}
-	else if (OnLocalTransactionFinalizedDelegate.IsBound())
+
+	void CaptureFinalized(FConcertExportedObject& ObjectUpdate, const FConcertObjectId& ObjectId, UObject* InObject)
 	{
-		const bool bIsNewlyCreated = InTransactionEvent.HasPendingKillChange() && IsValid(InObject);
+		const bool bIsNewlyCreated = TransactionEvent.HasPendingKillChange() && IsValid(InObject);
 
-		FConcertExportedObject& ObjectUpdate = OngoingTransaction.FinalizedData.FinalizedObjectUpdates.AddDefaulted_GetRef();
 		ObjectUpdate.ObjectId = ObjectId;
 		ObjectUpdate.ObjectPathDepth = ConcertSyncClientUtil::GetObjectPathDepth(InObject);
 		ObjectUpdate.ObjectData.bAllowCreate = bIsNewlyCreated && bCanCreateObject;
@@ -968,13 +1045,14 @@ void FConcertClientTransactionBridge::HandleObjectTransacted(UObject* InObject,
 		ObjectUpdate.ObjectData.NewOuterPathName = NewObjectOuterPathName;
 		ObjectUpdate.ObjectData.NewExternalPackageName = NewObjectExternalPackageName;
 
+
 		if (TransactionAnnotation && bUseSerializedAnnotationData)
 		{
 			FConcertSyncObjectWriter AnnotationWriter(&OngoingTransaction.FinalizedData.FinalizedLocalIdentifierTable, InObject, ObjectUpdate.SerializedAnnotationData, bIncludeEditorOnlyProperties, false);
 			TransactionAnnotation->Serialize(AnnotationWriter);
 		}
 
-		if (bIncludeNonPropertyObjectData && InTransactionEvent.HasNonPropertyChanges(/*SerializationOnly*/true))
+		if (bIncludeNonPropertyObjectData && TransactionEvent.HasNonPropertyChanges(/*SerializationOnly*/true))
 		{
 			// The 'non-property changes' refers to custom data added by a deriver UObject before and/or after the standard serialized data. Since this is a custom
 			// data format, we don't know what changed, call the object to re-serialize this part, but still send the delta for the generic reflected properties (in RootPropertyNames).
@@ -990,10 +1068,156 @@ void FConcertClientTransactionBridge::HandleObjectTransacted(UObject* InObject,
 		}
 		else // Its possible to optimize the transaction payload, only sending a 'delta' update.
 		{
-			// Only send properties that changed. The receiving side will 'patch' the object using the reflection system. The specific object serialization function will NOT be called.
 			ConcertSyncClientUtil::SerializeProperties(&OngoingTransaction.FinalizedData.FinalizedLocalIdentifierTable, InObject, ExportedProperties, bIncludeEditorOnlyProperties, ObjectUpdate.PropertyDatas);
 		}
 	}
+
+	void TrackPackageChanges(UPackage* ChangedPackage)
+	{
+		// Track which packages were changed
+		OngoingTransaction.CommonData.ModifiedPackages.AddUnique(ChangedPackage->GetFName());
+
+		// if there was an outer change, track that outer package
+		if (TransactionEvent.HasOuterChange())
+		{
+			FName OriginalOuterPackageName = *FPackageName::ObjectPathToPackageName(TransactionEvent.GetOriginalObjectOuterPathName().ToString());
+			OngoingTransaction.CommonData.ModifiedPackages.AddUnique(OriginalOuterPackageName);
+		}
+
+		// if there was an package change, track that package
+		if (TransactionEvent.HasExternalPackageChange())
+		{
+			FName OriginalPackageName = TransactionEvent.GetOriginalObjectPackageName().IsNone() ? *FPackageName::ObjectPathToPackageName(TransactionEvent.GetOriginalObjectOuterPathName().ToString()) : TransactionEvent.GetOriginalObjectPackageName();
+			OngoingTransaction.CommonData.ModifiedPackages.AddUnique(OriginalPackageName);
+		}
+	}
+
+	bool IsSnapshotEvent() const
+	{
+		return TransactionEvent.GetEventType() == ETransactionObjectEventType::Snapshot
+			&& (ExportedProperties.Num() > 0 || TransactionAnnotation.IsValid());
+	}
+};
+
+void AppendLevelInstanceObjectsToExportedObject(UObject* InObject, FConcertExportedObject& ExportedObject)
+{
+	ExportedObject.LevelInstanceObjects.Reset();
+	auto AddLevelInstanceObject = [&ExportedObject](UObject* InObject)
+	{
+		ExportedObject.LevelInstanceObjects.Add(FConcertObjectId(InObject));
+	};
+
+	ApplyForAllRelevantObjects(InObject, AddLevelInstanceObject);
+}
+
+void AppendLevelInstanceObjects(const TMap<FConcertObjectId, FWeakObjectPtr>& InObjectPtrs, TArray<FConcertExportedObject>& ExportedObjects)
+{
+	for (FConcertExportedObject& ExportedObject : ExportedObjects)
+	{
+		const FWeakObjectPtr* ObjectPtr = InObjectPtrs.Find(ExportedObject.ObjectId);
+		if (ObjectPtr && ObjectPtr->IsValid())
+		{
+			AppendLevelInstanceObjectsToExportedObject(ObjectPtr->Get(), ExportedObject);
+		}
+	}
+}
+
+}
+
+bool FConcertClientTransactionBridge::CanHandleObjectTransacted(UObject* InObject, const FTransactionObjectEvent& InTransactionEvent) const
+{
+	if (bIgnoreLocalTransactions)
+	{
+		return false;
+	}
+
+	if (!bIncludeAnnotationObjectChanges && InTransactionEvent.GetObjectChangeCreatedBy() == ETransactionObjectChangeCreatedBy::TransactionAnnotation)
+	{
+		return false;
+	}
+
+	if (!bIncludeEditorOnlyProperties && InObject->IsEditorOnly())
+	{
+		return false;
+	}
+
+	return true;
+}
+
+void FConcertClientTransactionBridge::HandleObjectTransacted(UObject* InObject, const FTransactionObjectEvent& InTransactionEvent)
+{
+	LLM_SCOPE_BYTAG(Concert_ConcertClientTransactionBridge);
+	SCOPED_CONCERT_TRACE(FConcertClientTransactionBridge_HandleObjectTransacted);
+
+	if (!CanHandleObjectTransacted(InObject, InTransactionEvent))
+	{
+		return;
+	}
+
+	UPackage* ChangedPackage = InObject->GetOutermost();
+	ETransactionFilterResult FilterResult = ConcertClientTransactionBridgeUtil::ApplyTransactionFilters(TransactionFilters, InObject, ChangedPackage);
+	FOngoingTransaction* TrackedTransaction = OngoingTransactions.Find(InTransactionEvent.GetOperationId());
+
+	// TODO: This needs to send both editor-only and non-editor-only payload
+	// data to the server, which will forward only the correct part to cooked
+	// and non-cooked clients
+	{
+		ConcertClientTransactionBridgeUtil::LogTransactionEventDetails(InObject, InTransactionEvent, FilterResult, TrackedTransaction != nullptr);
+	}
+
+	if (TrackedTransaction == nullptr)
+	{
+		// No tracked transaction so nothing to handle.
+		return;
+	}
+
+	const FConcertObjectId ObjectId = FConcertObjectId(InTransactionEvent.GetOriginalObjectId(), InObject->GetFlags());
+	FOngoingTransaction& OngoingTransaction = *TrackedTransaction;
+
+	// If the object is excluded or exclude the whole transaction add it to the excluded list
+	if (FilterResult != ETransactionFilterResult::IncludeObject)
+	{
+		OngoingTransaction.CommonData.bIsExcluded |= FilterResult == ETransactionFilterResult::ExcludeTransaction;
+		OngoingTransaction.CommonData.ExcludedObjectUpdates.Add(ObjectId);
+		return;
+	}
+
+	ConcertClientTransactionBridgeUtil::FTransactedObjectState
+		TransactedObjectState(InObject, InTransactionEvent, OngoingTransaction, bIncludeNonPropertyObjectData, bIncludeEditorOnlyProperties,  bIncludeAnnotationObjectChanges);
+	if (!TransactedObjectState)
+	{
+		// This object has no changes to send
+		return;
+	}
+
+	TransactedObjectState.TrackPackageChanges(ChangedPackage);
+
+	// Add this object change to its pending transaction
+	if (OnLocalTransactionSnapshotDelegate.IsBound() && TransactedObjectState.IsSnapshotEvent())
+	{
+		// Find or add an entry for this object
+		FConcertExportedObject* ObjectUpdatePtr = OngoingTransaction.SnapshotData.SnapshotObjectUpdates.FindByPredicate([&ObjectId](FConcertExportedObject& ObjectUpdate)
+		{
+			return ConcertSyncClientUtil::ObjectIdsMatch(ObjectId, ObjectUpdate.ObjectId);
+		});
+
+		if (!ObjectUpdatePtr)
+		{
+			ObjectUpdatePtr = &OngoingTransaction.SnapshotData.SnapshotObjectUpdates.AddDefaulted_GetRef();
+			ObjectUpdatePtr->ObjectId = ObjectId;
+			ObjectUpdatePtr->ObjectPathDepth = ConcertSyncClientUtil::GetObjectPathDepth(InObject);
+			ObjectUpdatePtr->ObjectData.bIsPendingKill = !IsValid(InObject);
+
+			OngoingTransaction.ObjectPtrs.FindOrAdd(ObjectId,InObject);
+		}
+		TransactedObjectState.CaptureSnapshot(InObject, ObjectUpdatePtr);
+	}
+	else if (OnLocalTransactionFinalizedDelegate.IsBound())
+	{
+		FConcertExportedObject& ObjectUpdate = OngoingTransaction.FinalizedData.FinalizedObjectUpdates.AddDefaulted_GetRef();
+		OngoingTransaction.ObjectPtrs.FindOrAdd(ObjectId,InObject);
+		TransactedObjectState.CaptureFinalized(ObjectUpdate, ObjectId, InObject);
+	}
 }
 
 void FConcertClientTransactionBridge::ConditionalBindUnderlyingLocalTransactionEvents()
@@ -1045,8 +1269,11 @@ void FConcertClientTransactionBridge::OnEndFrame()
 			continue;
 		}
 
+
 		if (OngoingTransactionPtr->bIsFinalized)
 		{
+			ConcertClientTransactionBridgeUtil::AppendLevelInstanceObjects(OngoingTransactionPtr->ObjectPtrs, OngoingTransactionPtr->FinalizedData.FinalizedObjectUpdates);
+
 			OnLocalTransactionFinalizedDelegate.Broadcast(OngoingTransactionPtr->CommonData, OngoingTransactionPtr->FinalizedData);
 
 			OngoingTransactions.Remove(OngoingTransactionPtr->CommonData.OperationId);
@@ -1055,6 +1282,8 @@ void FConcertClientTransactionBridge::OnEndFrame()
 		}
 		else if (OngoingTransactionPtr->SnapshotData.SnapshotObjectUpdates.Num() > 0)
 		{
+			ConcertClientTransactionBridgeUtil::AppendLevelInstanceObjects(OngoingTransactionPtr->ObjectPtrs, OngoingTransactionPtr->SnapshotData.SnapshotObjectUpdates);
+
 			OnLocalTransactionSnapshotDelegate.Broadcast(OngoingTransactionPtr->CommonData, OngoingTransactionPtr->SnapshotData);
 
 			OngoingTransactionPtr->bHasNotifiedSnapshot = true;
diff --git a/Engine/Plugins/Developer/Concert/ConcertSync/ConcertSyncClient/Source/ConcertSyncClient/Private/ConcertClientTransactionBridge.h b/Engine/Plugins/Developer/Concert/ConcertSync/ConcertSyncClient/Source/ConcertSyncClient/Private/ConcertClientTransactionBridge.h
--- a/Engine/Plugins/Developer/Concert/ConcertSync/ConcertSyncClient/Source/ConcertSyncClient/Private/ConcertClientTransactionBridge.h
+++ b/Engine/Plugins/Developer/Concert/ConcertSync/ConcertSyncClient/Source/ConcertSyncClient/Private/ConcertClientTransactionBridge.h
@@ -6,6 +6,11 @@
 #include "Misc/ITransaction.h"
 #include "ConcertSyncArchives.h"
 
+namespace ConcertClientTransactionBridgeUtil
+{
+	struct FTransactedObjectState;
+}
+
 class FConcertClientTransactionBridge : public IConcertClientTransactionBridge
 {
 public:
@@ -32,6 +37,9 @@ private:
 	/** Called to handle a transaction state change */
 	void HandleTransactionStateChanged(const FTransactionContext& InTransactionContext, const ETransactionStateEventType InTransactionState);
 
+	/** Confirms that we can handle the transaction event. */
+	bool CanHandleObjectTransacted(UObject *InObject, const FTransactionObjectEvent& InTransactionEvent) const;
+
 	/** Called to handle an object being transacted */
 	void HandleObjectTransacted(UObject* InObject, const FTransactionObjectEvent& InTransactionEvent);
 
@@ -54,10 +62,13 @@ private:
 		FConcertClientLocalTransactionCommonData CommonData;
 		FConcertClientLocalTransactionSnapshotData SnapshotData;
 		FConcertClientLocalTransactionFinalizedData FinalizedData;
+		TMap<FConcertObjectId, FWeakObjectPtr> ObjectPtrs;
 		bool bIsFinalized = false;
 		bool bHasNotifiedSnapshot = false;
 	};
 
+	friend struct ConcertClientTransactionBridgeUtil::FTransactedObjectState;
+
 	/** Array of transaction IDs in the order they should be notified (maps to OngoingTransactions, although canceled transactions may be missing from the map) */
 	TArray<FGuid> OngoingTransactionsOrder;
 
diff --git a/Engine/Plugins/Developer/Concert/ConcertSync/ConcertSyncClient/Source/ConcertSyncClient/Private/ConcertClientWorkspace.cpp b/Engine/Plugins/Developer/Concert/ConcertSync/ConcertSyncClient/Source/ConcertSyncClient/Private/ConcertClientWorkspace.cpp
--- a/Engine/Plugins/Developer/Concert/ConcertSync/ConcertSyncClient/Source/ConcertSyncClient/Private/ConcertClientWorkspace.cpp
+++ b/Engine/Plugins/Developer/Concert/ConcertSync/ConcertSyncClient/Source/ConcertSyncClient/Private/ConcertClientWorkspace.cpp
@@ -110,26 +110,81 @@ private:
 	TArray<TUniquePtr<FScopedSlowTask>> ExtendedTaskLife;
 };
 
-void SetReflectEditorLevelVisibilityWithGame(bool bInValue, bool bReset = false)
+namespace UE::ConcertWorkspace::Private
 {
-#if WITH_EDITOR
-	static bool bHasBeenSet = false;
-
-	if (bReset)
+	template <size_t N>
+	struct CVarString
 	{
-		bHasBeenSet = false;
-	}
+		TCHAR Storage[N+1]{};
+		constexpr CVarString(TCHAR const* InStr)
+		{
+			for (size_t Index = 0; Index != N; Index++)
+			{
+				Storage[Index] = InStr[Index];
+			}
+		}
+		constexpr operator TCHAR const*() const {return Storage;}
+	};
+	template<size_t N> CVarString(TCHAR const (&)[N]) -> CVarString<N - 1>;
 
-	// Detail mode was modified, so store in the CVar
-	static IConsoleVariable* CVar = IConsoleManager::Get().FindConsoleVariable(TEXT("Editor.ReflectEditorLevelVisibilityWithGame"));
-	int32 ValAsInt = !!bInValue;
-	bool bCanSet = !bHasBeenSet && CVar->GetInt() != ValAsInt;
-	if (GEditor && bCanSet)
+	template<CVarString T, bool bEditorOnly>
+	struct FCVarBool
 	{
-		CVar->Set(ValAsInt);
-		bHasBeenSet = bReset ? false : true;
-	}
+		static constexpr TCHAR const* CvarName = T;
+		void SetValueOnce(bool bInValue, bool bForce = false)
+		{
+#if WITH_EDITOR
+			const bool bShouldSetValue = bEditorOnly ? GEditor != nullptr : true;
+			if (bShouldSetValue && (!bHasBeenSet || bForce))
+			{
+				int32 ValueAsInt = !!bInValue;
+				SetValue(ValueAsInt);
+			}
+#endif
+		}
+
+		void Reset()
+		{
+#if WITH_EDITOR
+			if (bHasBeenSet)
+			{
+				SetValue(bOriginalValue);
+				bHasBeenSet = false;
+			}
 #endif
+		}
+
+		bool IsEnabled() const
+		{
+#if WITH_EDITOR
+			static IConsoleVariable* CVar = IConsoleManager::Get().FindConsoleVariable(CvarName);
+			if (CVar)
+			{
+				return !!CVar->GetInt();
+			}
+#endif
+			return false;
+		}
+
+	private:
+		void SetValue(int32 InValue)
+		{
+			static IConsoleVariable* CVar = IConsoleManager::Get().FindConsoleVariable(CvarName);
+			if (CVar)
+			{
+				bOriginalValue = CVar->GetInt();
+				CVar->Set(InValue);
+				bHasBeenSet = true;
+			}
+		}
+
+		int32 bOriginalValue = false;
+		bool bHasBeenSet = false;
+	};
+
+	static FCVarBool<TEXT("Editor.ReflectEditorLevelVisibilityWithGame"), true /*Editor Only*/> ReflectVisInGame;
+	static FCVarBool<TEXT("LevelInstance.ForceEditorWorldMode"), false /*Editor Only*/> LevelInstanceForceEditorWorldMode;
+	static FCVarBool<TEXT("EditorPaths.Enabled"), false /*Editor Only*/> EditorPaths;
 }
 
 struct FConcertWorkspaceConsoleCommands
@@ -640,6 +695,14 @@ void FConcertClientWorkspace::HandleConnectionChanged(IConcertClientSession& InS
 				}
 			}
 		}
+		if (GEditor == nullptr)
+		{
+			const bool bLevelInstanceEditorWorldMode = UE::ConcertWorkspace::Private::LevelInstanceForceEditorWorldMode.IsEnabled();
+			UE_CLOG( !bLevelInstanceEditorWorldMode, LogConcert, Warning, TEXT("Level Instance Editor World Mode is off. Ensure LevelInstance.ForceEditorWorldMode=1 on startup. Non-editor world mode in -game may cause issues with transaction playback on Level Instance Actors. Forcing this mode to be ON.") );
+		}
+
+		UE::ConcertWorkspace::Private::LevelInstanceForceEditorWorldMode.SetValueOnce(true);
+		UE::ConcertWorkspace::Private::EditorPaths.SetValueOnce(true);
 #endif
 	}
 	else if (Status == EConcertConnectionStatus::Disconnected)
@@ -647,7 +710,9 @@ void FConcertClientWorkspace::HandleConnectionChanged(IConcertClientSession& InS
 		bHasSyncedWorkspace = false;
 		bFinalizeWorkspaceSyncRequested = false;
 		FConcertSlowTaskStackWorkaround::Get().PopTask(MoveTemp(InitialSyncSlowTask));
-		SetReflectEditorLevelVisibilityWithGame(false, true);
+		UE::ConcertWorkspace::Private::ReflectVisInGame.Reset();
+		UE::ConcertWorkspace::Private::LevelInstanceForceEditorWorldMode.Reset();
+		UE::ConcertWorkspace::Private::EditorPaths.Reset();
 	}
 }
 
@@ -870,7 +935,7 @@ void FConcertClientWorkspace::OnEndFrame()
 	}
 	LiveSession->GetSessionDatabase().UpdateAsynchronousTasks();
 	IConcertClientRef ConcertClient = OwnerSyncClient->GetConcertClient();
-	SetReflectEditorLevelVisibilityWithGame(ConcertClient->GetConfiguration()->ClientSettings.bReflectLevelEditorInGame);
+	UE::ConcertWorkspace::Private::ReflectVisInGame.SetValueOnce(ConcertClient->GetConfiguration()->ClientSettings.bReflectLevelEditorInGame);
 }
 
 void FConcertClientWorkspace::HandleWorkspaceSyncEndpointEvent(const FConcertSessionContext& Context, const FConcertWorkspaceSyncEndpointEvent& Event)
diff --git a/Engine/Plugins/Developer/Concert/ConcertSync/ConcertSyncCore/Source/ConcertSyncCore/Public/ConcertTransactionEvents.h b/Engine/Plugins/Developer/Concert/ConcertSync/ConcertSyncCore/Source/ConcertSyncCore/Public/ConcertTransactionEvents.h
--- a/Engine/Plugins/Developer/Concert/ConcertSync/ConcertSyncCore/Source/ConcertSyncCore/Public/ConcertTransactionEvents.h
+++ b/Engine/Plugins/Developer/Concert/ConcertSync/ConcertSyncCore/Source/ConcertSyncCore/Public/ConcertTransactionEvents.h
@@ -155,6 +155,9 @@ struct FConcertExportedObject
 	UPROPERTY()
 	int32 ObjectPathDepth = 0;
 
+	UPROPERTY()
+	TArray<FConcertObjectId> LevelInstanceObjects;
+
 	UPROPERTY()
 	FConcertSerializedObjectData ObjectData;
 
diff --git a/Engine/Plugins/VirtualProduction/Switchboard/Source/Switchboard/switchboard/devices/ndisplay/plugin_ndisplay.py b/Engine/Plugins/VirtualProduction/Switchboard/Source/Switchboard/switchboard/devices/ndisplay/plugin_ndisplay.py
--- a/Engine/Plugins/VirtualProduction/Switchboard/Source/Switchboard/switchboard/devices/ndisplay/plugin_ndisplay.py
+++ b/Engine/Plugins/VirtualProduction/Switchboard/Source/Switchboard/switchboard/devices/ndisplay/plugin_ndisplay.py
@@ -1038,6 +1038,8 @@ class DevicenDisplay(DeviceUnreal):
 
         # Always tell Chaos to be deterministic
         dp_cvars.append('p.Chaos.Solver.Deterministic=1')
+        # Always enable LevelInstance World Editor Mode.
+        dp_cvars.append('LevelInstance.ForceEditorWorldMode=1')
 
         # mediaprofile
         mediaprofile_gamepath = DevicenDisplay.csettings["mediaprofile"].get_value(self.name)
