Fix Horde Issue # 383149

#rnx

diff --git a/Engine/Plugins/Developer/Concert/ConcertSync/ConcertSyncClient/Source/ConcertSyncClient/Private/ConcertClientTransactionBridge.cpp b/Engine/Plugins/Developer/Concert/ConcertSync/ConcertSyncClient/Source/ConcertSyncClient/Private/ConcertClientTransactionBridge.cpp
--- a/Engine/Plugins/Developer/Concert/ConcertSync/ConcertSyncClient/Source/ConcertSyncClient/Private/ConcertClientTransactionBridge.cpp
+++ b/Engine/Plugins/Developer/Concert/ConcertSync/ConcertSyncClient/Source/ConcertSyncClient/Private/ConcertClientTransactionBridge.cpp
@@ -288,15 +288,12 @@ struct FEditorTransactionNotification
 
 FConcertExportedObject MakeLevelInstanceExportedObject(const FConcertObjectId& LevelInstanceId, const FConcertExportedObject& SourceObject)
 {
-	FConcertExportedObject OutExportedObject(
-		LevelInstanceId,                // FConcertObjectId
-		SourceObject.ObjectPathDepth,   // int32 ObjectPathDepth
-		{},                             // No level instance objects.
-		SourceObject.ObjectData,
-		SourceObject.PropertyDatas,
-		{}                              // We don't support annotation data with Level Instances. All changes should be property deltas.
-		);
-
+	FConcertExportedObject OutExportedObject;
+	OutExportedObject.ObjectId = LevelInstanceId;
+	OutExportedObject.ObjectPathDepth = SourceObject.ObjectPathDepth;
+	OutExportedObject.ObjectData = SourceObject.ObjectData;
+	OutExportedObject.PropertyDatas = SourceObject.PropertyDatas;
+	// We don't support annotation data with Level Instances. All changes should be property deltas.
 	return OutExportedObject;
 }
 
