Refactor concert bridge to better handle component changes in level instance objects.  We were capturing the root component and applying that to corresponding components in level instance actors.  This was causing issues with xform properties which depend on having the correct local value.  So we have to capture the component in the event which means creating additional FConcertObjectExport to store the data in the transaction event.

#rnx

diff --git a/Engine/Plugins/Developer/Concert/ConcertSync/ConcertSyncClient/Source/ConcertSyncClient/Private/ConcertClientTransactionBridge.cpp b/Engine/Plugins/Developer/Concert/ConcertSync/ConcertSyncClient/Source/ConcertSyncClient/Private/ConcertClientTransactionBridge.cpp
--- a/Engine/Plugins/Developer/Concert/ConcertSync/ConcertSyncClient/Source/ConcertSyncClient/Private/ConcertClientTransactionBridge.cpp
+++ b/Engine/Plugins/Developer/Concert/ConcertSync/ConcertSyncClient/Source/ConcertSyncClient/Private/ConcertClientTransactionBridge.cpp
@@ -286,16 +286,6 @@ struct FEditorTransactionNotification
 };
 #endif
 
-FConcertExportedObject MakeLevelInstanceExportedObject(const FConcertObjectId& LevelInstanceId, const FConcertExportedObject& SourceObject)
-{
-	FConcertExportedObject OutExportedObject;
-	OutExportedObject.ObjectId = LevelInstanceId;
-	OutExportedObject.ObjectPathDepth = SourceObject.ObjectPathDepth;
-	OutExportedObject.ObjectData = SourceObject.ObjectData;
-	OutExportedObject.PropertyDatas = SourceObject.PropertyDatas;
-	// We don't support annotation data with Level Instances. All changes should be property deltas.
-	return OutExportedObject;
-}
 
 void ProcessTransactionEvent(const FConcertTransactionEventBase& InEvent, const FConcertSessionVersionInfo* InVersionInfo, const TArray<FName>& InPackagesToProcess, const FConcertLocalIdentifierTable* InLocalIdentifierTablePtr, const bool bIsSnapshot, const FConcertSyncWorldRemapper& WorldRemapper, const bool bIncludeEditorOnlyProperties)
 {
@@ -309,44 +299,26 @@ void ProcessTransactionEvent(const FConcertTransactionEventBase& InEvent, const
 	// --------------------------------------------------------------------------------------------------------------------
 	// Phase 0
 	// --------------------------------------------------------------------------------------------------------------------
-	int32 LevelInstanceObjects = 0;
 
 #if WITH_EDITOR
-	TArray<FConcertExportedObject> AddedLevelInstanceObjects;
-	// For -game instances, we need to gather all affected Level Instance Actors and apply the delta change to those
-	// objects.
-	if (!GIsEditor)
-	{
-		for (const FConcertExportedObject& ExportedObject : InEvent.ExportedObjects)
-		{
-			LevelInstanceObjects += ExportedObject.LevelInstanceObjects.Num();
-		}
-		AddedLevelInstanceObjects.Reserve(LevelInstanceObjects);
-		if (LevelInstanceObjects > 0 && ConcertSyncClientUtil::GetExternalPersistentWorld())
-		{
-			// This is a WP or OFPA level and we cannot apply transactions in those worlds due to limitations
-			// of using LevelInstance.ForceEditorWorldMode=1 and cell generation.  Future support for OFPA / WP in Multi-user
-			// with LevelInstances is TBD. Non WP and -game nodes can work with Level Instances.
-			//
-			return;
-		}
-	}
+	const bool bIsWorldPartitionWorldInGame = !GIsEditor && ConcertSyncClientUtil::IsWorldPartitionWorld();
 #endif
 
 	TArray<const FConcertExportedObject*, TInlineAllocator<32>> SortedExportedObjects;
 	{
-		SortedExportedObjects.Reserve(InEvent.ExportedObjects.Num() + LevelInstanceObjects);
+		SortedExportedObjects.Reserve(InEvent.ExportedObjects.Num());
 		for (const FConcertExportedObject& ExportedObject : InEvent.ExportedObjects)
 		{
 #if WITH_EDITOR
-			// If the exported object is an actor in a level instance and we are not in -game then skip.
-			if (!GIsEditor)
+			// If the exported object in a level instance and we are not in
+			// -game or in a WP -game then skip processing. We can't process in
+			// -game with WP due to constraints with cell generation and in
+			// editor sessions the normal transaction flow will handle the proxy
+			// object.
+			//
+			if (ExportedObject.bHasLevelInstanceObject && (GIsEditor || bIsWorldPartitionWorldInGame))
 			{
-				for (const FConcertObjectId& Id : ExportedObject.LevelInstanceObjects)
-				{
-					AddedLevelInstanceObjects.Emplace(MakeLevelInstanceExportedObject(Id, ExportedObject));
-					SortedExportedObjects.Add(&(AddedLevelInstanceObjects.Last()));
-				}
+				continue;
 			}
 #endif
 			SortedExportedObjects.Add(&ExportedObject);
@@ -931,18 +903,14 @@ void GetCanCreateOrRenameObject(UObject* InObject, bool& bOutCanCreateObject, bo
 
 struct FTransactedObjectState
 {
-    bool bCanCreateObject = true;
-	bool bCanRenameObject = true;
-
 	using FOngoingTransaction = FConcertClientTransactionBridge::FOngoingTransaction;
 
-	const FTransactionObjectEvent& TransactionEvent;
-	FOngoingTransaction& OngoingTransaction;
-
 	FName NewObjectPackageName;
 	FName NewObjectName;
 	FName NewObjectOuterPathName;
 	FName NewObjectExternalPackageName;
+
+	FWeakObjectPtr MainObjectPtr;
 	TArray<const FProperty*> ExportedProperties;
 	TSharedPtr<ITransactionObjectAnnotation> TransactionAnnotation;
 	bool bUseSerializedAnnotationData;
@@ -950,14 +918,21 @@ struct FTransactedObjectState
 	bool bIncludeNonPropertyObjectData = false;
 	bool bIncludeEditorOnlyProperties = false;
 
+    bool bCanCreateObject = true;
+	bool bCanRenameObject = true;
+
+	// Transaction event flags
+	bool bHasIdOrPendingKillChanges = false;
+	bool bHasPendingKillChange = false;
+	bool bHasNonPropertyChanges = false;
+	bool bIsSnapshot = false;
+
 	FTransactedObjectState() = delete;
     FTransactedObjectState(UObject* InObject, const FTransactionObjectEvent& InTransactionEvent,
-						   FOngoingTransaction& InOngoingTransaction,
 						   bool bInIncludeNonPropertyObjectData,
 						   bool bInIncludeEditorOnlyProperties,
 						   bool bIncludeAnnotationObjectChanges)
-		: TransactionEvent(InTransactionEvent),
-		  OngoingTransaction(InOngoingTransaction)
+		: MainObjectPtr(InObject)
     {
 		bIncludeNonPropertyObjectData =  bInIncludeNonPropertyObjectData;
 		bIncludeEditorOnlyProperties = bInIncludeEditorOnlyProperties;
@@ -966,36 +941,41 @@ struct FTransactedObjectState
 
         if (bCanRenameObject)
         {
-            if (TransactionEvent.HasOuterChange() || TransactionEvent.HasExternalPackageChange())
+            if (InTransactionEvent.HasOuterChange() || InTransactionEvent.HasExternalPackageChange())
             {
                 NewObjectPackageName = InObject->GetPackage()->GetFName();
             }
-            if (TransactionEvent.HasNameChange())
+            if (InTransactionEvent.HasNameChange())
             {
                 NewObjectName = InObject->GetFName();
             }
-            if (TransactionEvent.HasOuterChange() && InObject->GetOuter())
+            if (InTransactionEvent.HasOuterChange() && InObject->GetOuter())
             {
                 NewObjectOuterPathName = FName(*InObject->GetOuter()->GetPathName()) ;
             }
-            if (TransactionEvent.HasExternalPackageChange() && InObject->GetExternalPackage())
+            if (InTransactionEvent.HasExternalPackageChange() && InObject->GetExternalPackage())
             {
                 NewObjectExternalPackageName = InObject->GetExternalPackage()->GetFName();
             }
         }
 
-        ExportedProperties = ConcertSyncClientUtil::GetExportedProperties(InObject->GetClass(), TransactionEvent.GetChangedProperties(), bIncludeEditorOnlyProperties);
+        ExportedProperties = ConcertSyncClientUtil::GetExportedProperties(InObject->GetClass(), InTransactionEvent.GetChangedProperties(), bIncludeEditorOnlyProperties);
 
-        TransactionAnnotation = TransactionEvent.GetAnnotation();
+		bHasIdOrPendingKillChanges = InTransactionEvent.HasIdOrPendingKillChanges();
+		bHasPendingKillChange = InTransactionEvent.HasPendingKillChange();
+		bHasNonPropertyChanges = InTransactionEvent.HasNonPropertyChanges(/*SerializationOnly*/true);
+		bIsSnapshot = InTransactionEvent.GetEventType() == ETransactionObjectEventType::Snapshot;
+
+        TransactionAnnotation = InTransactionEvent.GetAnnotation();
         bUseSerializedAnnotationData = !bIncludeAnnotationObjectChanges
             || !TransactionAnnotation
             || !TransactionAnnotation->SupportsAdditionalObjectChanges();
     }
 
-	explicit operator bool() const
+	bool IsValid() const
     {
-        const bool bObjectHasChangesToSend = TransactionEvent.HasIdOrPendingKillChanges()
-            || (bIncludeNonPropertyObjectData && TransactionEvent.HasNonPropertyChanges(/*SerializationOnly*/true))
+        const bool bObjectHasChangesToSend = bHasIdOrPendingKillChanges
+            || (bIncludeNonPropertyObjectData && bHasNonPropertyChanges)
             || ExportedProperties.Num() > 0
             || (TransactionAnnotation && bUseSerializedAnnotationData);
         return bObjectHasChangesToSend;
@@ -1028,15 +1008,15 @@ struct FTransactedObjectState
 		}
 	}
 
-	void CaptureFinalized(FConcertExportedObject& ObjectUpdate, const FConcertObjectId& ObjectId, UObject* InObject)
+	void CaptureFinalized(FConcertExportedObject& ObjectUpdate, FConcertLocalIdentifierTable& LocalIdentifierTable, const FConcertObjectId& ObjectId, UObject* InObject)
 	{
-		const bool bIsNewlyCreated = TransactionEvent.HasPendingKillChange() && IsValid(InObject);
+		const bool bIsNewlyCreated = bHasPendingKillChange && ::IsValid(InObject);
 
 		ObjectUpdate.ObjectId = ObjectId;
 		ObjectUpdate.ObjectPathDepth = ConcertSyncClientUtil::GetObjectPathDepth(InObject);
 		ObjectUpdate.ObjectData.bAllowCreate = bIsNewlyCreated && bCanCreateObject;
 		ObjectUpdate.ObjectData.bResetExisting = bIsNewlyCreated;
-		ObjectUpdate.ObjectData.bIsPendingKill = !IsValid(InObject);
+		ObjectUpdate.ObjectData.bIsPendingKill = !::IsValid(InObject);
 		ObjectUpdate.ObjectData.NewPackageName = NewObjectPackageName;
 		ObjectUpdate.ObjectData.NewName = NewObjectName;
 		ObjectUpdate.ObjectData.NewOuterPathName = NewObjectOuterPathName;
@@ -1045,15 +1025,15 @@ struct FTransactedObjectState
 
 		if (TransactionAnnotation && bUseSerializedAnnotationData)
 		{
-			FConcertSyncObjectWriter AnnotationWriter(&OngoingTransaction.FinalizedData.FinalizedLocalIdentifierTable, InObject, ObjectUpdate.SerializedAnnotationData, bIncludeEditorOnlyProperties, false);
+			FConcertSyncObjectWriter AnnotationWriter(&LocalIdentifierTable, InObject, ObjectUpdate.SerializedAnnotationData, bIncludeEditorOnlyProperties, false);
 			TransactionAnnotation->Serialize(AnnotationWriter);
 		}
 
-		if (bIncludeNonPropertyObjectData && TransactionEvent.HasNonPropertyChanges(/*SerializationOnly*/true))
+		if (bIncludeNonPropertyObjectData && bHasNonPropertyChanges)
 		{
 			// The 'non-property changes' refers to custom data added by a deriver UObject before and/or after the standard serialized data. Since this is a custom
 			// data format, we don't know what changed, call the object to re-serialize this part, but still send the delta for the generic reflected properties (in RootPropertyNames).
-			ConcertSyncClientUtil::SerializeObject(&OngoingTransaction.FinalizedData.FinalizedLocalIdentifierTable, InObject, &ExportedProperties, bIncludeEditorOnlyProperties, ObjectUpdate.ObjectData.SerializedData);
+			ConcertSyncClientUtil::SerializeObject(&LocalIdentifierTable, InObject, &ExportedProperties, bIncludeEditorOnlyProperties, ObjectUpdate.ObjectData.SerializedData);
 
 			// Track which properties changed. Not used to apply the transaction on the receiving side, the object specific serialization function will be used for that, but
 			// to be able to display, in the transaction detail view, which 'properties' changed in the transaction as transaction data is otherwise opaque to UI.
@@ -1065,58 +1045,68 @@ struct FTransactedObjectState
 		}
 		else // Its possible to optimize the transaction payload, only sending a 'delta' update.
 		{
-			ConcertSyncClientUtil::SerializeProperties(&OngoingTransaction.FinalizedData.FinalizedLocalIdentifierTable, InObject, ExportedProperties, bIncludeEditorOnlyProperties, ObjectUpdate.PropertyDatas);
+			ConcertSyncClientUtil::SerializeProperties(&LocalIdentifierTable, InObject, ExportedProperties, bIncludeEditorOnlyProperties, ObjectUpdate.PropertyDatas);
 		}
 	}
 
-	void TrackPackageChanges(UPackage* ChangedPackage)
+	void TrackPackageChanges(UPackage* ChangedPackage, FOngoingTransaction& OngoingTransaction, const FTransactionObjectEvent& InTransactionEvent)
 	{
 		// Track which packages were changed
 		OngoingTransaction.CommonData.ModifiedPackages.AddUnique(ChangedPackage->GetFName());
 
 		// if there was an outer change, track that outer package
-		if (TransactionEvent.HasOuterChange())
+		if (InTransactionEvent.HasOuterChange())
 		{
-			FName OriginalOuterPackageName = *FPackageName::ObjectPathToPackageName(TransactionEvent.GetOriginalObjectOuterPathName().ToString());
+			FName OriginalOuterPackageName = *FPackageName::ObjectPathToPackageName(InTransactionEvent.GetOriginalObjectOuterPathName().ToString());
 			OngoingTransaction.CommonData.ModifiedPackages.AddUnique(OriginalOuterPackageName);
 		}
 
 		// if there was an package change, track that package
-		if (TransactionEvent.HasExternalPackageChange())
+		if (InTransactionEvent.HasExternalPackageChange())
 		{
-			FName OriginalPackageName = TransactionEvent.GetOriginalObjectPackageName().IsNone() ? *FPackageName::ObjectPathToPackageName(TransactionEvent.GetOriginalObjectOuterPathName().ToString()) : TransactionEvent.GetOriginalObjectPackageName();
+			FName OriginalPackageName = InTransactionEvent.GetOriginalObjectPackageName().IsNone() ? *FPackageName::ObjectPathToPackageName(InTransactionEvent.GetOriginalObjectOuterPathName().ToString()) : InTransactionEvent.GetOriginalObjectPackageName();
 			OngoingTransaction.CommonData.ModifiedPackages.AddUnique(OriginalPackageName);
 		}
 	}
 
 	bool IsSnapshotEvent() const
 	{
-		return TransactionEvent.GetEventType() == ETransactionObjectEventType::Snapshot
-			&& (ExportedProperties.Num() > 0 || TransactionAnnotation.IsValid());
+		return bIsSnapshot && (ExportedProperties.Num() > 0 || TransactionAnnotation.IsValid());
 	}
-};
 
-void AppendLevelInstanceObjectsToExportedObject(UObject* InObject, FConcertExportedObject& ExportedObject)
-{
-	ExportedObject.LevelInstanceObjects.Reset();
-	auto AddLevelInstanceObject = [&ExportedObject](UObject* InObject)
+	TArray<FConcertExportedObject> AddLevelSequenceObjectsForFinalized(FConcertLocalIdentifierTable& LocalIdentifierTable)
 	{
-		ExportedObject.LevelInstanceObjects.Add(FConcertObjectId(InObject));
-	};
+		check(!IsSnapshotEvent());
+		TArray<FConcertExportedObject> ExportedObjects;
+		if (MainObjectPtr.IsValid())
+		{
+			auto Capture = [&LocalIdentifierTable, &ExportedObjects, this](UObject* InObject)
+			{
+				FConcertExportedObject& ExportedObject = ExportedObjects.AddDefaulted_GetRef();
+				ExportedObject.bHasLevelInstanceObject = true;
+				FConcertObjectId NewId = FConcertObjectId(InObject);
+				CaptureFinalized(ExportedObject, LocalIdentifierTable, NewId, InObject);
+			};
+			ApplyForAllRelevantObjects(MainObjectPtr.Get(), Capture);
+		}
+		return ExportedObjects;
+	}
+};
 
-	ApplyForAllRelevantObjects(InObject, AddLevelInstanceObject);
-}
+using FTransactedObjectStateMap = TMap<FConcertObjectId, TPimplPtr<FTransactedObjectState> >;
 
-void AppendLevelInstanceObjects(const TMap<FConcertObjectId, FWeakObjectPtr>& InObjectPtrs, TArray<FConcertExportedObject>& ExportedObjects)
+void AppendLevelInstanceObjects(FConcertLocalIdentifierTable& LocalIdentifierTable, const FTransactedObjectStateMap& InObjectPtrs, TArray<FConcertExportedObject>& ExportedObjects)
 {
+	TArray<FConcertExportedObject> NewExportedObjects;
 	for (FConcertExportedObject& ExportedObject : ExportedObjects)
 	{
-		const FWeakObjectPtr* ObjectPtr = InObjectPtrs.Find(ExportedObject.ObjectId);
-		if (ObjectPtr && ObjectPtr->IsValid())
+		const TPimplPtr<ConcertClientTransactionBridgeUtil::FTransactedObjectState>* ObjectState = InObjectPtrs.Find(ExportedObject.ObjectId);
+		if (ObjectState && ObjectState->IsValid())
 		{
-			AppendLevelInstanceObjectsToExportedObject(ObjectPtr->Get(), ExportedObject);
+			NewExportedObjects.Append(ObjectState->Get()->AddLevelSequenceObjectsForFinalized(LocalIdentifierTable));
 		}
 	}
+	ExportedObjects.Append(NewExportedObjects);
 }
 
 }
@@ -1179,41 +1169,53 @@ void FConcertClientTransactionBridge::HandleObjectTransacted(UObject* InObject,
 		return;
 	}
 
-	ConcertClientTransactionBridgeUtil::FTransactedObjectState
-		TransactedObjectState(InObject, InTransactionEvent, OngoingTransaction, bIncludeNonPropertyObjectData, bIncludeEditorOnlyProperties,  bIncludeAnnotationObjectChanges);
-	if (!TransactedObjectState)
+	FOngoingTransaction::FTransactedObjectStatePtr TransactedObjectState = MakePimpl<ConcertClientTransactionBridgeUtil::FTransactedObjectState>(InObject, InTransactionEvent, bIncludeNonPropertyObjectData, bIncludeEditorOnlyProperties,  bIncludeAnnotationObjectChanges);
+
+	if (!TransactedObjectState->IsValid())
 	{
 		// This object has no changes to send
 		return;
 	}
 
-	TransactedObjectState.TrackPackageChanges(ChangedPackage);
-
+	TransactedObjectState->TrackPackageChanges(ChangedPackage, OngoingTransaction, InTransactionEvent);
 	// Add this object change to its pending transaction
-	if (OnLocalTransactionSnapshotDelegate.IsBound() && TransactedObjectState.IsSnapshotEvent())
+	if (OnLocalTransactionSnapshotDelegate.IsBound() && TransactedObjectState->IsSnapshotEvent())
 	{
 		// Find or add an entry for this object
-		FConcertExportedObject* ObjectUpdatePtr = OngoingTransaction.SnapshotData.SnapshotObjectUpdates.FindByPredicate([&ObjectId](FConcertExportedObject& ObjectUpdate)
+		auto CaptureForSnapshot = [&TransactedObjectState, &OngoingTransaction, InObject, &ObjectId](UObject* NewObject)
 		{
-			return ConcertSyncClientUtil::ObjectIdsMatch(ObjectId, ObjectUpdate.ObjectId);
-		});
+			FConcertObjectId NewId = (NewObject == nullptr || NewObject == InObject) ? ObjectId : FConcertObjectId(NewObject);
+			FConcertExportedObject* ObjectUpdatePtr = OngoingTransaction.SnapshotData.SnapshotObjectUpdates.FindByPredicate([&NewId](FConcertExportedObject& ObjectUpdate)
+			{
+				return ConcertSyncClientUtil::ObjectIdsMatch(NewId, ObjectUpdate.ObjectId);
+			});
 
-		if (!ObjectUpdatePtr)
-		{
-			ObjectUpdatePtr = &OngoingTransaction.SnapshotData.SnapshotObjectUpdates.AddDefaulted_GetRef();
-			ObjectUpdatePtr->ObjectId = ObjectId;
-			ObjectUpdatePtr->ObjectPathDepth = ConcertSyncClientUtil::GetObjectPathDepth(InObject);
-			ObjectUpdatePtr->ObjectData.bIsPendingKill = !IsValid(InObject);
+			if (!ObjectUpdatePtr)
+			{
+				ObjectUpdatePtr = &OngoingTransaction.SnapshotData.SnapshotObjectUpdates.AddDefaulted_GetRef();
+				ObjectUpdatePtr->ObjectId = NewId;
+				ObjectUpdatePtr->ObjectPathDepth = ConcertSyncClientUtil::GetObjectPathDepth(InObject);
+				ObjectUpdatePtr->ObjectData.bIsPendingKill = !IsValid(InObject);
+				ObjectUpdatePtr->bHasLevelInstanceObject = NewObject != InObject;
+			}
+			TransactedObjectState->CaptureSnapshot(NewObject, ObjectUpdatePtr);
+		};
 
-			OngoingTransaction.ObjectPtrs.FindOrAdd(ObjectId,InObject);
-		}
-		TransactedObjectState.CaptureSnapshot(InObject, ObjectUpdatePtr);
+		CaptureForSnapshot(InObject);
+		// Now traverse the object to see if we need to add any additional objects to the snapshot.
+		//
+		ConcertClientTransactionBridgeUtil::ApplyForAllRelevantObjects(InObject, CaptureForSnapshot);
 	}
 	else if (OnLocalTransactionFinalizedDelegate.IsBound())
 	{
+		FConcertLocalIdentifierTable& LocalIdentifierTable = OngoingTransaction.FinalizedData.FinalizedLocalIdentifierTable;
 		FConcertExportedObject& ObjectUpdate = OngoingTransaction.FinalizedData.FinalizedObjectUpdates.AddDefaulted_GetRef();
-		OngoingTransaction.ObjectPtrs.FindOrAdd(ObjectId,InObject);
-		TransactedObjectState.CaptureFinalized(ObjectUpdate, ObjectId, InObject);
+		TransactedObjectState->CaptureFinalized(ObjectUpdate, LocalIdentifierTable, ObjectId, InObject);
+
+		// We don't scan for additional object to add to the level snapshot due to undo flow.  On undo, with level instances,
+		// we cannot access actor list.
+		//
+		OngoingTransaction.TransactedStatePtrs.FindOrAdd(ObjectId,MoveTemp(TransactedObjectState));
 	}
 }
 
@@ -1269,7 +1271,7 @@ void FConcertClientTransactionBridge::OnEndFrame()
 
 		if (OngoingTransactionPtr->bIsFinalized)
 		{
-			ConcertClientTransactionBridgeUtil::AppendLevelInstanceObjects(OngoingTransactionPtr->ObjectPtrs, OngoingTransactionPtr->FinalizedData.FinalizedObjectUpdates);
+			ConcertClientTransactionBridgeUtil::AppendLevelInstanceObjects(OngoingTransactionPtr->FinalizedData.FinalizedLocalIdentifierTable, OngoingTransactionPtr->TransactedStatePtrs, OngoingTransactionPtr->FinalizedData.FinalizedObjectUpdates);
 
 			OnLocalTransactionFinalizedDelegate.Broadcast(OngoingTransactionPtr->CommonData, OngoingTransactionPtr->FinalizedData);
 
@@ -1279,8 +1281,6 @@ void FConcertClientTransactionBridge::OnEndFrame()
 		}
 		else if (OngoingTransactionPtr->SnapshotData.SnapshotObjectUpdates.Num() > 0)
 		{
-			ConcertClientTransactionBridgeUtil::AppendLevelInstanceObjects(OngoingTransactionPtr->ObjectPtrs, OngoingTransactionPtr->SnapshotData.SnapshotObjectUpdates);
-
 			OnLocalTransactionSnapshotDelegate.Broadcast(OngoingTransactionPtr->CommonData, OngoingTransactionPtr->SnapshotData);
 
 			OngoingTransactionPtr->bHasNotifiedSnapshot = true;
diff --git a/Engine/Plugins/Developer/Concert/ConcertSync/ConcertSyncClient/Source/ConcertSyncClient/Private/ConcertClientTransactionBridge.h b/Engine/Plugins/Developer/Concert/ConcertSync/ConcertSyncClient/Source/ConcertSyncClient/Private/ConcertClientTransactionBridge.h
--- a/Engine/Plugins/Developer/Concert/ConcertSync/ConcertSyncClient/Source/ConcertSyncClient/Private/ConcertClientTransactionBridge.h
+++ b/Engine/Plugins/Developer/Concert/ConcertSync/ConcertSyncClient/Source/ConcertSyncClient/Private/ConcertClientTransactionBridge.h
@@ -62,7 +62,10 @@ private:
 		FConcertClientLocalTransactionCommonData CommonData;
 		FConcertClientLocalTransactionSnapshotData SnapshotData;
 		FConcertClientLocalTransactionFinalizedData FinalizedData;
-		TMap<FConcertObjectId, FWeakObjectPtr> ObjectPtrs;
+
+		using FTransactedObjectStatePtr = TPimplPtr<ConcertClientTransactionBridgeUtil::FTransactedObjectState>;
+		TMap<FConcertObjectId, FTransactedObjectStatePtr> TransactedStatePtrs;
+
 		bool bIsFinalized = false;
 		bool bHasNotifiedSnapshot = false;
 	};
diff --git a/Engine/Plugins/Developer/Concert/ConcertSync/ConcertSyncClient/Source/ConcertSyncClient/Private/ConcertSyncClientUtil.cpp b/Engine/Plugins/Developer/Concert/ConcertSync/ConcertSyncClient/Source/ConcertSyncClient/Private/ConcertSyncClientUtil.cpp
--- a/Engine/Plugins/Developer/Concert/ConcertSync/ConcertSyncClient/Source/ConcertSyncClient/Private/ConcertSyncClientUtil.cpp
+++ b/Engine/Plugins/Developer/Concert/ConcertSync/ConcertSyncClient/Source/ConcertSyncClient/Private/ConcertSyncClientUtil.cpp
@@ -747,6 +747,17 @@ ULevel* GetExternalPersistentWorld()
 #endif
 	return nullptr;
 }
+bool IsWorldPartitionWorld()
+{
+#if WITH_EDITOR
+	UWorld* OwningWorld = GetCurrentWorld();
+	if (OwningWorld)
+	{
+		return OwningWorld->GetWorldPartition() != nullptr;
+	}
+#endif
+	return false;
+}
 
 void FillPackageInfo(UPackage* InPackage, UObject* InAsset, const EConcertPackageUpdateType InPackageUpdateType, FConcertPackageInfo& OutPackageInfo)
 {
diff --git a/Engine/Plugins/Developer/Concert/ConcertSync/ConcertSyncClient/Source/ConcertSyncClient/Private/ConcertSyncClientUtil.h b/Engine/Plugins/Developer/Concert/ConcertSync/ConcertSyncClient/Source/ConcertSyncClient/Private/ConcertSyncClientUtil.h
--- a/Engine/Plugins/Developer/Concert/ConcertSync/ConcertSyncClient/Source/ConcertSyncClient/Private/ConcertSyncClientUtil.h
+++ b/Engine/Plugins/Developer/Concert/ConcertSync/ConcertSyncClient/Source/ConcertSyncClient/Private/ConcertSyncClientUtil.h
@@ -143,5 +143,8 @@ namespace ConcertSyncClientUtil
 	/** Return the Persistent Level if it is currently using external objects, i.e. one file per actor for world partitioning. */
 	ULevel* GetExternalPersistentWorld();
 
+	/** Return true if we are in a world partition world. */
+	bool IsWorldPartitionWorld();
+
 	void FillPackageInfo(UPackage* InPackage, UObject* InAsset, const EConcertPackageUpdateType InPackageUpdateType, FConcertPackageInfo& OutPackageInfo);
 }
diff --git a/Engine/Plugins/Developer/Concert/ConcertSync/ConcertSyncCore/Source/ConcertSyncCore/Public/ConcertTransactionEvents.h b/Engine/Plugins/Developer/Concert/ConcertSync/ConcertSyncCore/Source/ConcertSyncCore/Public/ConcertTransactionEvents.h
--- a/Engine/Plugins/Developer/Concert/ConcertSync/ConcertSyncCore/Source/ConcertSyncCore/Public/ConcertTransactionEvents.h
+++ b/Engine/Plugins/Developer/Concert/ConcertSync/ConcertSyncCore/Source/ConcertSyncCore/Public/ConcertTransactionEvents.h
@@ -156,7 +156,7 @@ struct FConcertExportedObject
 	int32 ObjectPathDepth = 0;
 
 	UPROPERTY()
-	TArray<FConcertObjectId> LevelInstanceObjects;
+	bool bHasLevelInstanceObject = false;
 
 	UPROPERTY()
 	FConcertSerializedObjectData ObjectData;
