Check for take data when applying pending take on the client.

diff --git a/Engine/Plugins/Developer/Concert/ConcertSync/ConcertSyncClient/Source/ConcertSyncClient/Private/ConcertClientSequencerManager.cpp b/Engine/Plugins/Developer/Concert/ConcertSync/ConcertSyncClient/Source/ConcertSyncClient/Private/ConcertClientSequencerManager.cpp
--- a/Engine/Plugins/Developer/Concert/ConcertSync/ConcertSyncClient/Source/ConcertSyncClient/Private/ConcertClientSequencerManager.cpp
+++ b/Engine/Plugins/Developer/Concert/ConcertSync/ConcertSyncClient/Source/ConcertSyncClient/Private/ConcertClientSequencerManager.cpp
@@ -770,12 +770,20 @@ void FConcertClientSequencerManager::ApplyOpenEvent(const FConcertSequencerOpenE
 				UE_LOG(LogConcertSequencerSync, Display, TEXT("Missing Object %s when loading PendingTake"), *FString(MissingObject));
 			});
 
-		FConcertSyncWorldRemapper Remapper(
-			TEXT("/Engine/Transient.__PendingLevelSequence__"), PendingLevelSequence->GetPathName());
-		FConcertSyncObjectReader Reader(nullptr, MoveTemp(Remapper), nullptr, PendingLevelSequence, InOpenEvent.TakeData.Bytes,
-										MissingObjectDelegate);
-		Reader.SetSerializeNestedObjects(true);
-		Reader.SerializeObject(PendingLevelSequence);
+		if (InOpenEvent.TakeData.Bytes.Num() > 0)
+		{
+			FConcertSyncWorldRemapper Remapper(
+				TEXT("/Engine/Transient.__PendingLevelSequence__"), PendingLevelSequence->GetPathName());
+			FConcertSyncObjectReader Reader(nullptr, MoveTemp(Remapper), nullptr, PendingLevelSequence, InOpenEvent.TakeData.Bytes,
+											MissingObjectDelegate);
+			Reader.SetSerializeNestedObjects(true);
+			Reader.SerializeObject(PendingLevelSequence);
+		}
+		else
+		{
+			UE_LOG(LogConcertSequencerSync, Display, TEXT("Missing take data on pending take."));
+		}
+
 		#if WITH_EDITOR
 		if (GIsEditor)
 		{
diff --git a/Engine/Plugins/Developer/Concert/ConcertSync/ConcertSyncServer/Source/ConcertSyncServer/Private/ConcertServerSequencerManager.cpp b/Engine/Plugins/Developer/Concert/ConcertSync/ConcertSyncServer/Source/ConcertSyncServer/Private/ConcertServerSequencerManager.cpp
--- a/Engine/Plugins/Developer/Concert/ConcertSync/ConcertSyncServer/Source/ConcertSyncServer/Private/ConcertServerSequencerManager.cpp
+++ b/Engine/Plugins/Developer/Concert/ConcertSync/ConcertSyncServer/Source/ConcertSyncServer/Private/ConcertServerSequencerManager.cpp
@@ -73,7 +73,7 @@ void FConcertServerSequencerManager::HandleSequencerOpenEvent(const FConcertSess
 	FConcertOpenSequencerState& SequencerState = SequencerStates.FindOrAdd(*InEvent.SequenceObjectPath);
 	SequencerState.ClientEndpointIds.AddUnique(InEventContext.SourceEndpointId);
 	SequencerState.State.SequenceObjectPath = InEvent.SequenceObjectPath;
-	SequencerState.TakeData = InEvent.TakeData;
+	SequencerState.TakeData = InEvent.TakeData.Bytes.Num() > 0 ? InEvent.TakeData : SequencerState.TakeData;
 
 	// Forward the message to the other clients
 	TArray<FGuid> ClientIds = LiveSession->GetSession().GetSessionClientEndpointIds();
