Remove check for SnapshotEvent. Not all users subscribe to the snapshot event and thus it is not true that we will not be a snapshot event when processing a transaction. Also adjust logic for snapshot / finalized.  We were capturing an extra transaction.

diff --git a/Engine/Plugins/Developer/Concert/ConcertSync/ConcertSyncClient/Source/ConcertSyncClient/Private/ConcertClientTransactionBridge.cpp b/Engine/Plugins/Developer/Concert/ConcertSync/ConcertSyncClient/Source/ConcertSyncClient/Private/ConcertClientTransactionBridge.cpp
--- a/Engine/Plugins/Developer/Concert/ConcertSync/ConcertSyncClient/Source/ConcertSyncClient/Private/ConcertClientTransactionBridge.cpp
+++ b/Engine/Plugins/Developer/Concert/ConcertSync/ConcertSyncClient/Source/ConcertSyncClient/Private/ConcertClientTransactionBridge.cpp
@@ -1069,14 +1069,13 @@ struct FTransactedObjectState
 		}
 	}
 
-	bool IsSnapshotEvent() const
+	bool HasDataToCapture() const
 	{
-		return bIsSnapshot && (ExportedProperties.Num() > 0 || TransactionAnnotation.IsValid());
+		return (ExportedProperties.Num() > 0 || TransactionAnnotation.IsValid());
 	}
 
 	TArray<FConcertExportedObject> AddLevelSequenceObjectsForFinalized(FConcertLocalIdentifierTable& LocalIdentifierTable)
 	{
-		check(!IsSnapshotEvent());
 		TArray<FConcertExportedObject> ExportedObjects;
 		if (MainObjectPtr.IsValid())
 		{
@@ -1179,32 +1178,35 @@ void FConcertClientTransactionBridge::HandleObjectTransacted(UObject* InObject,
 
 	TransactedObjectState->TrackPackageChanges(ChangedPackage, OngoingTransaction, InTransactionEvent);
 	// Add this object change to its pending transaction
-	if (OnLocalTransactionSnapshotDelegate.IsBound() && TransactedObjectState->IsSnapshotEvent())
+	if (InTransactionEvent.GetEventType() == ETransactionObjectEventType::Snapshot)
 	{
-		// Find or add an entry for this object
-		auto CaptureForSnapshot = [&TransactedObjectState, &OngoingTransaction, InObject, &ObjectId](UObject* NewObject)
+		if (OnLocalTransactionSnapshotDelegate.IsBound() && TransactedObjectState->HasDataToCapture())
 		{
-			FConcertObjectId NewId = (NewObject == nullptr || NewObject == InObject) ? ObjectId : FConcertObjectId(NewObject);
-			FConcertExportedObject* ObjectUpdatePtr = OngoingTransaction.SnapshotData.SnapshotObjectUpdates.FindByPredicate([&NewId](FConcertExportedObject& ObjectUpdate)
+			// Find or add an entry for this object
+			auto CaptureForSnapshot = [&TransactedObjectState, &OngoingTransaction, InObject, &ObjectId](UObject* NewObject)
 			{
-				return ConcertSyncClientUtil::ObjectIdsMatch(NewId, ObjectUpdate.ObjectId);
-			});
+				FConcertObjectId NewId = (NewObject == nullptr || NewObject == InObject) ? ObjectId : FConcertObjectId(NewObject);
+				FConcertExportedObject* ObjectUpdatePtr = OngoingTransaction.SnapshotData.SnapshotObjectUpdates.FindByPredicate([&NewId](FConcertExportedObject& ObjectUpdate)
+					{
+						return ConcertSyncClientUtil::ObjectIdsMatch(NewId, ObjectUpdate.ObjectId);
+					});
 
-			if (!ObjectUpdatePtr)
-			{
-				ObjectUpdatePtr = &OngoingTransaction.SnapshotData.SnapshotObjectUpdates.AddDefaulted_GetRef();
-				ObjectUpdatePtr->ObjectId = NewId;
-				ObjectUpdatePtr->ObjectPathDepth = ConcertSyncClientUtil::GetObjectPathDepth(InObject);
-				ObjectUpdatePtr->ObjectData.bIsPendingKill = !IsValid(InObject);
-				ObjectUpdatePtr->bHasLevelInstanceObject = NewObject != InObject;
-			}
-			TransactedObjectState->CaptureSnapshot(NewObject, ObjectUpdatePtr);
-		};
+				if (!ObjectUpdatePtr)
+				{
+					ObjectUpdatePtr = &OngoingTransaction.SnapshotData.SnapshotObjectUpdates.AddDefaulted_GetRef();
+					ObjectUpdatePtr->ObjectId = NewId;
+					ObjectUpdatePtr->ObjectPathDepth = ConcertSyncClientUtil::GetObjectPathDepth(InObject);
+					ObjectUpdatePtr->ObjectData.bIsPendingKill = !IsValid(InObject);
+					ObjectUpdatePtr->bHasLevelInstanceObject = NewObject != InObject;
+				}
+				TransactedObjectState->CaptureSnapshot(NewObject, ObjectUpdatePtr);
+			};
 
-		CaptureForSnapshot(InObject);
-		// Now traverse the object to see if we need to add any additional objects to the snapshot.
-		//
-		ConcertClientTransactionBridgeUtil::ApplyForAllRelevantObjects(InObject, CaptureForSnapshot);
+			CaptureForSnapshot(InObject);
+			// Now traverse the object to see if we need to add any additional objects to the snapshot.
+			//
+			ConcertClientTransactionBridgeUtil::ApplyForAllRelevantObjects(InObject, CaptureForSnapshot);
+		}
 	}
 	else if (OnLocalTransactionFinalizedDelegate.IsBound())
 	{
