Added custom object factory support to ConcertSyncClient

This allows it to understand how to correctly create/destroy object types that it cannot directly depend on due to dependency limitations

#rnx

diff --git a/Engine/Plugins/Developer/Concert/ConcertSync/ConcertSyncClient/Source/ConcertSyncClient/Private/ConcertClientObjectFactory.cpp b/Engine/Plugins/Developer/Concert/ConcertSync/ConcertSyncClient/Source/ConcertSyncClient/Private/ConcertClientObjectFactory.cpp
new file mode 100644
--- /dev/null
+++ b/Engine/Plugins/Developer/Concert/ConcertSync/ConcertSyncClient/Source/ConcertSyncClient/Private/ConcertClientObjectFactory.cpp
@@ -0,0 +1,68 @@
+// Copyright Epic Games, Inc. All Rights Reserved.
+
+#include "ConcertClientObjectFactory.h"
+
+class FConcertClientObjectFactoryCache
+{
+public:
+	static FConcertClientObjectFactoryCache& Get()
+	{
+		static FConcertClientObjectFactoryCache Instance;
+		return Instance;
+	}
+
+	void RegisterFactory(const UConcertClientObjectFactory* Factory)
+	{
+		checkf(!Factory->GetClass()->HasAnyClassFlags(CLASS_Abstract), TEXT("Factory '%s' is abstract and cannot be registered!"), *Factory->GetPathName());
+		checkf(!Factories.Contains(Factory), TEXT("Factory '%s' was already registered!"), *Factory->GetPathName());
+		Factories.Add(Factory);
+	}
+
+	void UnregisterFactory(const UConcertClientObjectFactory* Factory)
+	{
+		Factories.Remove(Factory);
+	}
+
+	const UConcertClientObjectFactory* FindFactoryForClass(const UClass* Class)
+	{
+		for (const UConcertClientObjectFactory* Factory : Factories)
+		{
+			if (Factory->SupportsClass(Class))
+			{
+				return Factory;
+			}
+		}
+
+		return nullptr;
+	}
+
+private:
+	TArray<const UConcertClientObjectFactory*> Factories;
+};
+
+void UConcertClientObjectFactory::PostInitProperties()
+{
+	Super::PostInitProperties();
+
+	checkf(HasAnyFlags(RF_ClassDefaultObject), TEXT("Factories are assumed to be CDOs! If you need to support multiple instances of the same factory class then this code will need updating accordingly."));
+
+	if (!GetClass()->HasAnyClassFlags(CLASS_Abstract))
+	{
+		FConcertClientObjectFactoryCache::Get().RegisterFactory(this);
+	}
+}
+
+void UConcertClientObjectFactory::BeginDestroy()
+{
+	Super::BeginDestroy();
+
+	if (!GetClass()->HasAnyClassFlags(CLASS_Abstract))
+	{
+		FConcertClientObjectFactoryCache::Get().UnregisterFactory(this);
+	}
+}
+
+const UConcertClientObjectFactory* UConcertClientObjectFactory::FindFactoryForClass(const UClass* Class)
+{
+	return FConcertClientObjectFactoryCache::Get().FindFactoryForClass(Class);
+}
diff --git a/Engine/Plugins/Developer/Concert/ConcertSync/ConcertSyncClient/Source/ConcertSyncClient/Private/ConcertClientTransactionBridge.cpp b/Engine/Plugins/Developer/Concert/ConcertSync/ConcertSyncClient/Source/ConcertSyncClient/Private/ConcertClientTransactionBridge.cpp
--- a/Engine/Plugins/Developer/Concert/ConcertSync/ConcertSyncClient/Source/ConcertSyncClient/Private/ConcertClientTransactionBridge.cpp
+++ b/Engine/Plugins/Developer/Concert/ConcertSync/ConcertSyncClient/Source/ConcertSyncClient/Private/ConcertClientTransactionBridge.cpp
@@ -1,6 +1,7 @@
 // Copyright Epic Games, Inc. All Rights Reserved.
 
 #include "ConcertClientTransactionBridge.h"
+#include "ConcertClientObjectFactory.h"
 #include "ConcertLogGlobal.h"
 #include "ConcertSyncSettings.h"
 #include "ConcertSyncClientUtil.h"
@@ -343,6 +344,7 @@ void ProcessTransactionEvent(const FConcertTransactionEventBase& InEvent, const
 	TArray<ConcertSyncClientUtil::FGetObjectResult, TInlineAllocator<32>> TransactionObjects;
 	{
 		TSet<const UObject*> NewlyCreatedObjects;
+		TSortedMap<const UConcertClientObjectFactory*, TArray<UObject*>> ObjectsPendingInitialization;
 
 		// Find or create each object in parent->child order
 		TransactionObjects.AddDefaulted(SortedExportedObjects.Num());
@@ -370,6 +372,12 @@ void ProcessTransactionEvent(const FConcertTransactionEventBase& InEvent, const
 			{
 				if (TransactionObjectRef.NewlyCreated())
 				{
+					// If this object was created from a factory, then it may need further initialization once every object has been created
+					if (TransactionObjectRef.Factory)
+					{
+						ObjectsPendingInitialization.FindOrAdd(TransactionObjectRef.Factory).Add(TransactionObjectRef.Obj);
+					}
+
 					// If this is a new component then we need to apply its CreationMethod (if present) early, as it could affect the PreEdit behavior
 					if (UActorComponent* Component = Cast<UActorComponent>(TransactionObjectRef.Obj))
 					{
@@ -435,6 +443,12 @@ void ProcessTransactionEvent(const FConcertTransactionEventBase& InEvent, const
 #endif
 			}
 		}
+
+		// Run any deferred object factory initialization
+		for (const TTuple<const UConcertClientObjectFactory*, TArray<UObject*>>& ObjectsPendingInitializationPair : ObjectsPendingInitialization)
+		{
+			ObjectsPendingInitializationPair.Key->InitializeObjects(ObjectsPendingInitializationPair.Value);
+		}
 	}
 
 #if WITH_EDITOR
diff --git a/Engine/Plugins/Developer/Concert/ConcertSync/ConcertSyncClient/Source/ConcertSyncClient/Private/ConcertSyncClientUtil.cpp b/Engine/Plugins/Developer/Concert/ConcertSync/ConcertSyncClient/Source/ConcertSyncClient/Private/ConcertSyncClientUtil.cpp
--- a/Engine/Plugins/Developer/Concert/ConcertSync/ConcertSyncClient/Source/ConcertSyncClient/Private/ConcertSyncClientUtil.cpp
+++ b/Engine/Plugins/Developer/Concert/ConcertSync/ConcertSyncClient/Source/ConcertSyncClient/Private/ConcertSyncClientUtil.cpp
@@ -2,6 +2,7 @@
 
 #include "ConcertSyncClientUtil.h"
 #include "ConcertSyncArchives.h"
+#include "ConcertClientObjectFactory.h"
 #include "ConcertTransactionEvents.h"
 #include "ConcertLogGlobal.h"
 #include "ConcertSyncSettings.h"
@@ -124,26 +125,34 @@ void UpdatePendingKillState(UObject* InObj, const bool bIsPendingKill)
 
 	if (bIsPendingKill)
 	{
-		bool bMarkAsGarbage = true;
+		bool bDestructionHandled = false;
 
-		if (AActor* Actor = Cast<AActor>(InObj))
+		if (const UConcertClientObjectFactory* Factory = UConcertClientObjectFactory::FindFactoryForClass(InObj->GetClass()))
 		{
-			if (UWorld* ActorWorld = Actor->GetWorld())
+			bDestructionHandled = Factory->DestroyObject(InObj);
+		}
+
+		if (!bDestructionHandled)
+		{
+			if (AActor* Actor = Cast<AActor>(InObj))
 			{
-#if WITH_EDITOR
-				if (GIsEditor)
+				if (UWorld* ActorWorld = Actor->GetWorld())
 				{
-					bMarkAsGarbage = !ActorWorld->EditorDestroyActor(Actor, /*bShouldModifyLevel*/false);
-				}
-				else
+#if WITH_EDITOR
+					if (GIsEditor)
+					{
+						bDestructionHandled = ActorWorld->EditorDestroyActor(Actor, /*bShouldModifyLevel*/false);
+					}
+					else
 #endif	// WITH_EDITOR
-				{
-					bMarkAsGarbage = !ActorWorld->DestroyActor(Actor, /*bNetForce*/false, /*bShouldModifyLevel*/false);
+					{
+						bDestructionHandled = ActorWorld->DestroyActor(Actor, /*bNetForce*/false, /*bShouldModifyLevel*/false);
+					}
 				}
 			}
 		}
 
-		if (bMarkAsGarbage)
+		if (!bDestructionHandled)
 		{
 			InObj->MarkAsGarbage();
 		}
@@ -347,59 +356,74 @@ FGetObjectResult GetObject(const FConcertObjectId& InObjectId, const FName InNew
 			if (bAllowCreate)
 			{
 				FGetObjectResult ObjectResult;
+				ObjectResult.Factory = UConcertClientObjectFactory::FindFactoryForClass(ObjectClass);
+
 				// Create the new object
-				if (ObjectClass->IsChildOf<AActor>())
+				bool bFactoryHandledCreation = false;
+				if (ObjectResult.Factory)
 				{
-					// Actors should go through SpawnActor where possible
-					if (ULevel* OuterLevel = Cast<ULevel>(NewObjectOuter))
+					bFactoryHandledCreation = ObjectResult.Factory->CreateObject(ObjectResult.Obj, NewObjectOuter, ObjectClass, *ObjectNameToCreate.ToString(), (EObjectFlags)InObjectId.ObjectPersistentFlags);
+				}
+				if (!bFactoryHandledCreation)
+				{
+					if (ObjectClass->IsChildOf<AActor>())
 					{
-						UWorld* OwnerWorld = OuterLevel->GetWorld();
-						if (!OwnerWorld)
+						// Actors should go through SpawnActor where possible
+						if (ULevel* OuterLevel = Cast<ULevel>(NewObjectOuter))
 						{
-							OwnerWorld = OuterLevel->GetTypedOuter<UWorld>();
-						}
+							UWorld* OwnerWorld = OuterLevel->GetWorld();
+							if (!OwnerWorld)
+							{
+								OwnerWorld = OuterLevel->GetTypedOuter<UWorld>();
+							}
 
-						if (OwnerWorld)
-						{
-							UObject* ExistingObjectOfDifferentClass = StaticFindObjectFast(nullptr, OuterLevel, ObjectNameToCreate);
-							if (!ExistingObjectOfDifferentClass)
+							if (OwnerWorld)
 							{
-								FActorSpawnParameters SpawnParams;
-								SpawnParams.Name = ObjectNameToCreate;
-								SpawnParams.OverrideLevel = OuterLevel;
-								SpawnParams.SpawnCollisionHandlingOverride = ESpawnActorCollisionHandlingMethod::AlwaysSpawn;
-								SpawnParams.bNoFail = true;
-								SpawnParams.ObjectFlags = (EObjectFlags)InObjectId.ObjectPersistentFlags;
-								ObjectResult = FGetObjectResult(OwnerWorld->SpawnActor<AActor>(ObjectClass, FVector::ZeroVector, FRotator::ZeroRotator, SpawnParams), EGetObjectResultFlags::NewlyCreated);
+								UObject* ExistingObjectOfDifferentClass = StaticFindObjectFast(nullptr, OuterLevel, ObjectNameToCreate);
+								if (!ExistingObjectOfDifferentClass)
+								{
+									FActorSpawnParameters SpawnParams;
+									SpawnParams.Name = ObjectNameToCreate;
+									SpawnParams.OverrideLevel = OuterLevel;
+									SpawnParams.SpawnCollisionHandlingOverride = ESpawnActorCollisionHandlingMethod::AlwaysSpawn;
+									SpawnParams.bNoFail = true;
+									SpawnParams.ObjectFlags = (EObjectFlags)InObjectId.ObjectPersistentFlags;
+									ObjectResult.Obj = OwnerWorld->SpawnActor<AActor>(ObjectClass, FVector::ZeroVector, FRotator::ZeroRotator, SpawnParams);
+								}
+								else
+								{
+									UE_LOG(LogConcert, Warning, TEXT("Actor '%s' already exists! Expected class: '%s'"), *ExistingObjectOfDifferentClass->GetFullName(), *ObjectClass->GetPathName());
+									ensureMsgf(!ExistingObjectOfDifferentClass, TEXT("Actor '%s' already exists! Expected class: '%s'"), *ExistingObjectOfDifferentClass->GetFullName(), *ObjectClass->GetPathName());
+								}
 							}
 							else
 							{
-								UE_LOG(LogConcert, Warning, TEXT("Actor '%s' already exists! Expected class: '%s'"), *ExistingObjectOfDifferentClass->GetFullName(), *ObjectClass->GetPathName()); 
-								ensureMsgf(!ExistingObjectOfDifferentClass, TEXT("Actor '%s' already exists! Expected class: '%s'"), *ExistingObjectOfDifferentClass->GetFullName(), *ObjectClass->GetPathName());
+								UE_LOG(LogConcert, Warning, TEXT("Actor '%s' could not find an owner World! This is unexpected and the Actor will be created via NewObject rather than SpawnActor."), *ObjectNameToCreate.ToString());
 							}
 						}
 						else
 						{
-							UE_LOG(LogConcert, Warning, TEXT("Actor '%s' could not find an owner World! This is unexpected and the Actor will be created via NewObject rather than SpawnActor."), *ObjectNameToCreate.ToString());
+							UE_LOG(LogConcert, Warning, TEXT("Actor '%s' wasn't directly outered to a Level! This is unexpected and the Actor will be created via NewObject rather than SpawnActor."), *ObjectNameToCreate.ToString());
 						}
 					}
 					else
 					{
-						UE_LOG(LogConcert, Warning, TEXT("Actor '%s' wasn't directly outered to a Level! This is unexpected and the Actor will be created via NewObject rather than SpawnActor."), *ObjectNameToCreate.ToString());
+						ObjectResult.Obj = NewObject<UObject>(NewObjectOuter, ObjectClass, *ObjectNameToCreate.ToString(), (EObjectFlags)InObjectId.ObjectPersistentFlags);
+
+						if (UActorComponent* NewComponent = Cast<UActorComponent>(ObjectResult.Obj))
+						{
+							NewComponent->RegisterComponent();
+						}
 					}
 				}
-				else
+				
+				if (ObjectResult.Obj)
 				{
-					ObjectResult = FGetObjectResult(NewObject<UObject>(NewObjectOuter, ObjectClass, *ObjectNameToCreate.ToString(), (EObjectFlags)InObjectId.ObjectPersistentFlags), EGetObjectResultFlags::NewlyCreated);
+					// if we have any package assignment, do it here
+					AssignExternalPackage(ObjectResult.Obj);
 
-					if (UActorComponent* NewComponent = Cast<UActorComponent>(ObjectResult.Obj))
-					{
-						NewComponent->RegisterComponent();
-					}
+					ObjectResult.Flags |= EGetObjectResultFlags::NewlyCreated;
 				}
-				
-				// if we have any package assignment, do it here
-				AssignExternalPackage(ObjectResult.Obj);
 
 				return ObjectResult;
 			}
diff --git a/Engine/Plugins/Developer/Concert/ConcertSync/ConcertSyncClient/Source/ConcertSyncClient/Private/ConcertSyncClientUtil.h b/Engine/Plugins/Developer/Concert/ConcertSync/ConcertSyncClient/Source/ConcertSyncClient/Private/ConcertSyncClientUtil.h
--- a/Engine/Plugins/Developer/Concert/ConcertSync/ConcertSyncClient/Source/ConcertSyncClient/Private/ConcertSyncClientUtil.h
+++ b/Engine/Plugins/Developer/Concert/ConcertSync/ConcertSyncClient/Source/ConcertSyncClient/Private/ConcertSyncClientUtil.h
@@ -20,6 +20,7 @@ class UActorComponent;
 
 class FConcertSyncWorldRemapper;
 class FConcertLocalIdentifierTable;
+class UConcertClientObjectFactory;
 
 struct FConcertObjectId;
 struct FConcertWorldNodeId;
@@ -41,14 +42,17 @@ namespace ConcertSyncClientUtil
 
 	struct FGetObjectResult
 	{
-		FGetObjectResult()
-			: Obj(nullptr)
-			, Flags(EGetObjectResultFlags::None)
+		FGetObjectResult() = default;
+
+		explicit FGetObjectResult(UObject* InObj, const EGetObjectResultFlags InFlags = EGetObjectResultFlags::None)
+			: Obj(InObj)
+			, Flags(InFlags)
 		{
 		}
 
-		explicit FGetObjectResult(UObject* InObj, const EGetObjectResultFlags InFlags = EGetObjectResultFlags::None)
+		explicit FGetObjectResult(UObject* InObj, const UConcertClientObjectFactory* InFactory, const EGetObjectResultFlags InFlags)
 			: Obj(InObj)
+			, Factory(InFactory)
 			, Flags(InFlags)
 		{
 		}
@@ -63,8 +67,9 @@ namespace ConcertSyncClientUtil
 			return EnumHasAnyFlags(Flags, EGetObjectResultFlags::NewlyCreated);
 		}
 
-		UObject* Obj;
-		EGetObjectResultFlags Flags;
+		UObject* Obj = nullptr;
+		const UConcertClientObjectFactory* Factory = nullptr;
+		EGetObjectResultFlags Flags = EGetObjectResultFlags::None;
 	};
 
 	bool IsUserEditing();
diff --git a/Engine/Plugins/Developer/Concert/ConcertSync/ConcertSyncClient/Source/ConcertSyncClient/Public/ConcertClientObjectFactory.h b/Engine/Plugins/Developer/Concert/ConcertSync/ConcertSyncClient/Source/ConcertSyncClient/Public/ConcertClientObjectFactory.h
new file mode 100644
--- /dev/null
+++ b/Engine/Plugins/Developer/Concert/ConcertSync/ConcertSyncClient/Source/ConcertSyncClient/Public/ConcertClientObjectFactory.h
@@ -0,0 +1,63 @@
+// Copyright Epic Games, Inc. All Rights Reserved.
+
+#pragma once
+
+#include "Containers/ArrayView.h"
+#include "UObject/Object.h"
+
+#include "ConcertClientObjectFactory.generated.h"
+
+/**
+ * Factory API that can be implemented to extend the set of objects that Concert understands how to create/destroy.
+ * @note This API is const as it operates on CDOs, which also means that a factory implementation cannot store any additional state.
+ */
+UCLASS(MinimalAPI, Abstract)
+class UConcertClientObjectFactory : public UObject
+{
+	GENERATED_BODY()
+	
+public:
+	//~ UObject
+	CONCERTSYNCCLIENT_API virtual void PostInitProperties() override;
+	CONCERTSYNCCLIENT_API virtual void BeginDestroy() override;
+
+	/**
+	 * Attempt to find a factory that supports the given class.
+	 * @note If multiple factories support the given class then the result isn't guaranteed to be deterministic.
+	 */
+	static CONCERTSYNCCLIENT_API const UConcertClientObjectFactory* FindFactoryForClass(const UClass* Class);
+
+	/**
+	 * Does this factory support creating/destroying objects of the given class?
+	 */
+	virtual bool SupportsClass(const UClass* Class) const
+	{
+		return false;
+	}
+	
+	/**
+	 * Attempt to create a new object from the given arguments, and fill in OutObject.
+	 * @return True if this function handled the creation attempt (this doesn't mean that OutObject isn't null!), or false if we should fallback to using NewObject.
+	 */
+	virtual bool CreateObject(UObject*& OutObject, UObject* Outer, const UClass* Class, const FName Name, const EObjectFlags Flags) const
+	{
+		return false;
+	}
+	
+	/**
+	 * Perform any additional initialization on the given objects.
+	 * @note This is called after CreateObject has been called for every object in the transaction, and can be used for delayed initialization of interdependent objects.
+	 */
+	virtual void InitializeObjects(TArrayView<UObject* const> Objects) const
+	{
+	}
+	
+	/**
+	 * Attempt to destroy the given object.
+	 * @return True if this function handled the destruction attempt, or false if we should fallback to using MarkAsGarbage.
+	 */
+	virtual bool DestroyObject(UObject* Object) const
+	{
+		return false;
+	}
+};
