Improve error logging of requests and responses for stream and authority using CVars

diff --git a/Engine/Plugins/Developer/Concert/ConcertApp/MultiUserClient/Source/MultiUserClient/MultiUserClient.Build.cs b/Engine/Plugins/Developer/Concert/ConcertApp/MultiUserClient/Source/MultiUserClient/MultiUserClient.Build.cs
--- a/Engine/Plugins/Developer/Concert/ConcertApp/MultiUserClient/Source/MultiUserClient/MultiUserClient.Build.cs
+++ b/Engine/Plugins/Developer/Concert/ConcertApp/MultiUserClient/Source/MultiUserClient/MultiUserClient.Build.cs
@@ -25,6 +25,7 @@ namespace UnrealBuildTool.Rules
 					"EditorStyle",
 					"Engine",
 					"InputCore",
+					"JsonUtilities",
 					"MessageLog",
 					"Projects",
 					"Slate",
diff --git a/Engine/Plugins/Developer/Concert/ConcertApp/MultiUserClient/Source/MultiUserClient/Private/Replication/Submission/Remote/SubmissionWorkflow_RemoteClient.cpp b/Engine/Plugins/Developer/Concert/ConcertApp/MultiUserClient/Source/MultiUserClient/Private/Replication/Submission/Remote/SubmissionWorkflow_RemoteClient.cpp
--- a/Engine/Plugins/Developer/Concert/ConcertApp/MultiUserClient/Source/MultiUserClient/Private/Replication/Submission/Remote/SubmissionWorkflow_RemoteClient.cpp
+++ b/Engine/Plugins/Developer/Concert/ConcertApp/MultiUserClient/Source/MultiUserClient/Private/Replication/Submission/Remote/SubmissionWorkflow_RemoteClient.cpp
@@ -6,6 +6,7 @@
 #include "IConcertSession.h"
 #include "RemoteSubmissionMessages.h"
 #include "Replication/Submission/SingleClientSubmissionOperation.h"
+#include "Replication/Util/StreamRequestUtils.h"
 #include "Widgets/ActiveSession/Replication/Client/ClientUtils.h"
 
 namespace UE::MultiUserClient
@@ -188,8 +189,6 @@ namespace UE::MultiUserClient
 		FSingleClientSubmissionOperation& Operation = GetOperation();
 		if (EventData.bProcessedSuccessfully)
 		{
-			LogStreamErrorsIfNeeded(EventData.StreamChangeResponse);
-			
 			const FCompletedChangeSubmission StreamChangeSubmission { *InProgressOperation->Parameters.StreamRequest, { EventData.StreamChangeResponse } };
 			const FSubmitStreamChangesResponse StreamResponse { EStreamSubmissionErrorCode::Success, StreamChangeSubmission };
 			Operation.EmplaceStreamPromise(StreamResponse);
@@ -249,15 +248,4 @@ namespace UE::MultiUserClient
 	{
 		return ClientUtils::GetClientDisplayName(*ConcertSession, RemoteClientEndpointId);
 	}
-
-	void FSubmissionWorkflow_RemoteClient::LogStreamErrorsIfNeeded(const FConcertReplication_ChangeStream_Response& Response) const
-	{
-		if (Response.IsFailure())
-		{
-			FStringOutputDevice OutputDevice;
-			Response.LogErrors(OutputDevice);
-			const FString ClientName = ClientUtils::GetClientDisplayName(*ConcertSession, RemoteClientEndpointId);
-			UE_LOG(LogConcert, Warning, TEXT("Submission failed for remote client %s. Errors: %s"), *ClientName, *OutputDevice);
-		}
-	}
 }
diff --git a/Engine/Plugins/Developer/Concert/ConcertApp/MultiUserClient/Source/MultiUserClient/Private/Replication/Submission/Remote/SubmissionWorkflow_RemoteClient.h b/Engine/Plugins/Developer/Concert/ConcertApp/MultiUserClient/Source/MultiUserClient/Private/Replication/Submission/Remote/SubmissionWorkflow_RemoteClient.h
--- a/Engine/Plugins/Developer/Concert/ConcertApp/MultiUserClient/Source/MultiUserClient/Private/Replication/Submission/Remote/SubmissionWorkflow_RemoteClient.h
+++ b/Engine/Plugins/Developer/Concert/ConcertApp/MultiUserClient/Source/MultiUserClient/Private/Replication/Submission/Remote/SubmissionWorkflow_RemoteClient.h
@@ -82,8 +82,6 @@ namespace UE::MultiUserClient
 		void FinishSubmissionIfDone();
 
 		FString GetRemoteClientName() const;
-		
-		void LogStreamErrorsIfNeeded(const FConcertReplication_ChangeStream_Response& Response) const;
 	};
 }
 
diff --git a/Engine/Plugins/Developer/Concert/ConcertApp/MultiUserClient/Source/MultiUserClient/Private/Replication/Submission/SubmissionWorkflow_LocalClient.cpp b/Engine/Plugins/Developer/Concert/ConcertApp/MultiUserClient/Source/MultiUserClient/Private/Replication/Submission/SubmissionWorkflow_LocalClient.cpp
--- a/Engine/Plugins/Developer/Concert/ConcertApp/MultiUserClient/Source/MultiUserClient/Private/Replication/Submission/SubmissionWorkflow_LocalClient.cpp
+++ b/Engine/Plugins/Developer/Concert/ConcertApp/MultiUserClient/Source/MultiUserClient/Private/Replication/Submission/SubmissionWorkflow_LocalClient.cpp
@@ -4,6 +4,7 @@
 
 #include "IConcertSyncClient.h"
 #include "Replication/IConcertClientReplicationManager.h"
+#include "Replication/Util/StreamRequestUtils.h"
 #include "Widgets/ActiveSession/Replication/Client/ClientUtils.h"
 
 namespace UE::MultiUserClient
@@ -12,11 +13,6 @@ namespace UE::MultiUserClient
 		: Client(MoveTemp(InClient))
 	{}
 
-	FSubmissionWorkflow_LocalClient::~FSubmissionWorkflow_LocalClient()
-	{
-		
-	}
-
 	EChangeUploadability FSubmissionWorkflow_LocalClient::GetUploadability() const
 	{
 		const bool bOperationInProgress = InProgressOperation.IsSet(); 
@@ -77,7 +73,6 @@ namespace UE::MultiUserClient
 		)
 	{
 		const TSharedRef<FSingleClientSubmissionOperation> Operation = InProgressOperation->Operation;
-		LogStreamErrorsIfNeeded(ChangeStreamResponse);
 		
 		const EStreamSubmissionErrorCode ErrorCode = ChangeStreamResponse.ErrorCode == EReplicationResponseErrorCode::Handled
 			? EStreamSubmissionErrorCode::Success
@@ -155,16 +150,4 @@ namespace UE::MultiUserClient
 		Operation->EmplaceCompleteOperationPromise(ESubmissionOperationCompletedCode::Processed);
 		OnSubmitOperationCompletedDelegate.Broadcast();
 	}
-
-	void FSubmissionWorkflow_LocalClient::LogStreamErrorsIfNeeded(const FConcertReplication_ChangeStream_Response& Response) const
-	{
-		const TSharedPtr<IConcertClientSession> CurrentSession = Client->GetConcertClient()->GetCurrentSession();
-		if (Response.IsFailure() && CurrentSession)
-		{
-			FStringOutputDevice OutputDevice;
-			Response.LogErrors(OutputDevice);
-			const FString ClientName = ClientUtils::GetClientDisplayName(*CurrentSession, CurrentSession->GetSessionClientEndpointId());
-			UE_LOG(LogConcert, Warning, TEXT("Submission failed for local client %s. Errors: %s"), *ClientName, *OutputDevice);
-		}
-	}
 }
diff --git a/Engine/Plugins/Developer/Concert/ConcertApp/MultiUserClient/Source/MultiUserClient/Private/Replication/Submission/SubmissionWorkflow_LocalClient.h b/Engine/Plugins/Developer/Concert/ConcertApp/MultiUserClient/Source/MultiUserClient/Private/Replication/Submission/SubmissionWorkflow_LocalClient.h
--- a/Engine/Plugins/Developer/Concert/ConcertApp/MultiUserClient/Source/MultiUserClient/Private/Replication/Submission/SubmissionWorkflow_LocalClient.h
+++ b/Engine/Plugins/Developer/Concert/ConcertApp/MultiUserClient/Source/MultiUserClient/Private/Replication/Submission/SubmissionWorkflow_LocalClient.h
@@ -26,7 +26,6 @@ namespace UE::MultiUserClient
 	public:
 		
 		FSubmissionWorkflow_LocalClient(TSharedRef<IConcertSyncClient> InClient);
-		virtual ~FSubmissionWorkflow_LocalClient() override;
 		
 		//~ Begin ISubmissionWorkflow Interface
 		virtual TSharedPtr<ISubmissionOperation> SubmitChanges(FSubmissionParams Params) override;
@@ -61,8 +60,6 @@ namespace UE::MultiUserClient
 		
 		void SkipAuthorityStage(EAuthoritySubmissionRequestErrorCode RequestCode, EAuthoritySubmissionResponseErrorCode ResponseCode);
 		void CleanUpSubmissionOperation();
-
-		void LogStreamErrorsIfNeeded(const FConcertReplication_ChangeStream_Response& Response) const;
 	};
 }
 
diff --git a/Engine/Plugins/Developer/Concert/ConcertApp/MultiUserClient/Source/MultiUserClient/Private/Replication/Util/StreamRequestUtils.cpp b/Engine/Plugins/Developer/Concert/ConcertApp/MultiUserClient/Source/MultiUserClient/Private/Replication/Util/StreamRequestUtils.cpp
--- a/Engine/Plugins/Developer/Concert/ConcertApp/MultiUserClient/Source/MultiUserClient/Private/Replication/Util/StreamRequestUtils.cpp
+++ b/Engine/Plugins/Developer/Concert/ConcertApp/MultiUserClient/Source/MultiUserClient/Private/Replication/Util/StreamRequestUtils.cpp
@@ -2,8 +2,12 @@
 
 #include "StreamRequestUtils.h"
 
-#include "Replication/Stream/IClientStreamSynchronizer.h"
+#include "ConcertLogGlobal.h"
+#include "JsonObjectConverter.h"
 #include "Replication/Stream/StreamChangeTracker.h"
+#include "Widgets/ActiveSession/Replication/Client/ClientUtils.h"
+
+#include "Containers/UnrealString.h"
 
 namespace UE::MultiUserClient::StreamRequestUtils
 {
diff --git a/Engine/Plugins/Developer/Concert/ConcertApp/MultiUserClient/Source/MultiUserClient/Private/Replication/Util/StreamRequestUtils.h b/Engine/Plugins/Developer/Concert/ConcertApp/MultiUserClient/Source/MultiUserClient/Private/Replication/Util/StreamRequestUtils.h
--- a/Engine/Plugins/Developer/Concert/ConcertApp/MultiUserClient/Source/MultiUserClient/Private/Replication/Util/StreamRequestUtils.h
+++ b/Engine/Plugins/Developer/Concert/ConcertApp/MultiUserClient/Source/MultiUserClient/Private/Replication/Util/StreamRequestUtils.h
@@ -4,6 +4,8 @@
 
 #include "Replication/IConcertClientReplicationManager.h"
 
+class IConcertClientSession;
+
 struct FConcertReplication_ChangeStream_Request;
 
 namespace UE::MultiUserClient
diff --git a/Engine/Plugins/Developer/Concert/ConcertApp/MultiUserClient/Source/MultiUserClient/Private/Replication/Util/SynchronizedRequestUtils.cpp b/Engine/Plugins/Developer/Concert/ConcertApp/MultiUserClient/Source/MultiUserClient/Private/Replication/Util/SynchronizedRequestUtils.cpp
--- a/Engine/Plugins/Developer/Concert/ConcertApp/MultiUserClient/Source/MultiUserClient/Private/Replication/Util/SynchronizedRequestUtils.cpp
+++ b/Engine/Plugins/Developer/Concert/ConcertApp/MultiUserClient/Source/MultiUserClient/Private/Replication/Util/SynchronizedRequestUtils.cpp
@@ -27,20 +27,6 @@ namespace UE::MultiUserClient
 				}
 			}
 		}
-		
-		static void LogStreamChangeError(const FGuid& ClientId, const FSubmitStreamChangesResponse& Response)
-		{
-			if (Response.ErrorCode == EStreamSubmissionErrorCode::Timeout)
-			{
-				UE_LOG(LogConcert, Warning, TEXT("Remote stream change to client %s timed out."), *ClientId.ToString());
-			}
-			else if (Response.SubmissionInfo && Response.SubmissionInfo->Response.IsFailure())
-			{
-				FStringOutputDevice OutputDevice;
-				Response.SubmissionInfo->Response.LogErrors(OutputDevice);
-				UE_LOG(LogConcert, Warning, TEXT("Remote stream change to client %s failed: %s"), *ClientId.ToString(), *OutputDevice);
-			}
-		}
 
 		class FSyncOperation;
 
@@ -142,8 +128,6 @@ namespace UE::MultiUserClient
 
 			void OnCompleteStream(const FGuid& ClientId, FSubmitStreamChangesResponse&& Response)
 			{
-				LogStreamChangeError(ClientId, Response);
-				
 				IntermediateResults.StreamResponses[ClientId] = MoveTemp(Response);
 				FindOperationByClient(ClientId)->bHasCompletedStreamChanges.store(true, std::memory_order_release);
 
diff --git a/Engine/Plugins/Developer/Concert/ConcertSync/ConcertSyncClient/Source/ConcertSyncClient/ConcertSyncClient.Build.cs b/Engine/Plugins/Developer/Concert/ConcertSync/ConcertSyncClient/Source/ConcertSyncClient/ConcertSyncClient.Build.cs
--- a/Engine/Plugins/Developer/Concert/ConcertSync/ConcertSyncClient/Source/ConcertSyncClient/ConcertSyncClient.Build.cs
+++ b/Engine/Plugins/Developer/Concert/ConcertSync/ConcertSyncClient/Source/ConcertSyncClient/ConcertSyncClient.Build.cs
@@ -25,6 +25,7 @@ namespace UnrealBuildTool.Rules
 					"ConcertTransport",
 					"HeadMountedDisplay",
 					"InputCore",
+					"JsonUtilities",
 					"LevelSequence",
 					"MovieScene",
 					"RenderCore",
diff --git a/Engine/Plugins/Developer/Concert/ConcertSync/ConcertSyncClient/Source/ConcertSyncClient/Private/Replication/Manager/ReplicationManagerState_Connected.cpp b/Engine/Plugins/Developer/Concert/ConcertSync/ConcertSyncClient/Source/ConcertSyncClient/Private/Replication/Manager/ReplicationManagerState_Connected.cpp
--- a/Engine/Plugins/Developer/Concert/ConcertSync/ConcertSyncClient/Source/ConcertSyncClient/Private/Replication/Manager/ReplicationManagerState_Connected.cpp
+++ b/Engine/Plugins/Developer/Concert/ConcertSync/ConcertSyncClient/Source/ConcertSyncClient/Private/Replication/Manager/ReplicationManagerState_Connected.cpp
@@ -12,9 +12,9 @@
 #include "Replication/Processing/ClientReplicationDataCollector.h"
 #include "Replication/Processing/ObjectReplicationApplierProcessor.h"
 #include "Replication/Processing/ObjectReplicationReceiver.h"
-#include "Replication/Processing/ObjectReplicationSender.h"
 
 #include "Algo/RemoveIf.h"
+#include "JsonObjectConverter.h"
 #include "Misc/ScopeExit.h"
 #include "UObject/Package.h"
 #include "UObject/UObjectGlobals.h"
@@ -42,6 +42,31 @@ namespace UE::ConcertSyncClient::Replication
 		false,
 		TEXT("Whether the client should pretend that authority change requests were rejected.")
 		);
+
+	TAutoConsoleVariable<bool> CVarLogStreamRequestsAndResponses(
+		TEXT("Concert.Replication.LogStreamRequestsAndResponses"),
+		false,
+		TEXT("Whether to log changes to streams.")
+		);
+	TAutoConsoleVariable<bool> CVarLogAuthorityRequestsAndResponses(
+		TEXT("Concert.Replication.LogStreamRequestsAndResponses"),
+		false,
+		TEXT("Whether to log changes to authority.")
+		);
+
+	namespace Private
+	{
+		template<typename TMessage>
+		static void LogNetworkMessage(const TAutoConsoleVariable<bool>& ShouldLog, const TMessage& Message)
+		{
+			if (ShouldLog.GetValueOnAnyThread())
+			{
+				FString JsonString;
+				FJsonObjectConverter::UStructToJsonObjectString(TMessage::StaticStruct(), &Message, JsonString, 0, 0);
+				UE_LOG(LogConcert, Log, TEXT("%s\n%s"), *TMessage::StaticStruct()->GetName(), *JsonString);
+			}
+		}
+	}
 	
 	FReplicationManagerState_Connected::FReplicationManagerState_Connected(
 		TSharedRef<IConcertClientSession> LiveSession,
@@ -120,10 +145,13 @@ namespace UE::ConcertSyncClient::Replication
 		// Stop replicating removed objects right now: the server will remove authority after processing this request.
 		// At that point, it will log errors for receiving replication data from a client without authority.
 		HandleReleasingReplicatedObjects(Args);
-		
+
+		Private::LogNetworkMessage(CVarLogAuthorityRequestsAndResponses, Args);
 		return LiveSession->SendCustomRequest<FConcertReplication_ChangeAuthority_Request, FConcertReplication_ChangeAuthority_Response>(Args, LiveSession->GetSessionServerEndpointId())
 			.Next([WeakThis = TWeakPtr<FReplicationManagerState_Connected>(SharedThis(this)), Args](FConcertReplication_ChangeAuthority_Response&& Response) mutable
 			{
+				Private::LogNetworkMessage(CVarLogAuthorityRequestsAndResponses, Response);
+				
 				if (const TSharedPtr<FReplicationManagerState_Connected> ThisPin = WeakThis.Pin()
 					; ThisPin && Response.ErrorCode == EReplicationResponseErrorCode::Handled)
 				{
@@ -170,9 +198,12 @@ namespace UE::ConcertSyncClient::Replication
 		// At that point, it will log errors for receiving replication data from a client without authority.
 		HandleRemovingReplicatedObjects(Args);
 		
+		Private::LogNetworkMessage(CVarLogStreamRequestsAndResponses, Args);
 		return LiveSession->SendCustomRequest<FConcertReplication_ChangeStream_Request, FConcertReplication_ChangeStream_Response>(Args, LiveSession->GetSessionServerEndpointId())
 			.Next([WeakThis = TWeakPtr<FReplicationManagerState_Connected>(SharedThis(this)), Args](FConcertReplication_ChangeStream_Response&& Response)
 			{
+				Private::LogNetworkMessage(CVarLogStreamRequestsAndResponses, Response);
+
 				const TSharedPtr<FReplicationManagerState_Connected> ThisPin = WeakThis.Pin();
 				if (ThisPin && Response.IsSuccess())
 				{
diff --git a/Engine/Plugins/Developer/Concert/ConcertSync/ConcertSyncCore/Source/ConcertSyncCore/Private/Replication/Messages/ChangeStream.cpp b/Engine/Plugins/Developer/Concert/ConcertSync/ConcertSyncCore/Source/ConcertSyncCore/Private/Replication/Messages/ChangeStream.cpp
--- a/Engine/Plugins/Developer/Concert/ConcertSync/ConcertSyncCore/Source/ConcertSyncCore/Private/Replication/Messages/ChangeStream.cpp
+++ b/Engine/Plugins/Developer/Concert/ConcertSync/ConcertSyncCore/Source/ConcertSyncCore/Private/Replication/Messages/ChangeStream.cpp
@@ -2,8 +2,6 @@
 
 #include "Replication/Messages/ChangeStream.h"
 
-#include "Misc/OutputDevice.h"
-
 TOptional<FConcertReplication_ChangeStream_PutObject> FConcertReplication_ChangeStream_PutObject::MakeFromInfo(const FConcertReplicatedObjectInfo& New)
 {
 	if (!New.IsValidForSendingToServer())
@@ -31,31 +29,4 @@ TOptional<FConcertReplicatedObjectInfo> FConcertReplication_ChangeStream_PutObje
 		return {};
 	}
 	return FConcertReplicatedObjectInfo{ ClassPath, Properties };
-}
-
-void FConcertReplication_ChangeStream_Response::LogErrors(FOutputDevice& OutputDevice) const
-{
-	for (const TPair<FConcertObjectInStreamID, FConcertReplicatedObjectId>& Conflict : AuthorityConflicts)
-	{
-		OutputDevice.Logf(TEXT("Authority: { %s } conflicts with { %s }"), *Conflict.Key.ToString(), *Conflict.Value.ToString());
-	}
-
-	for (const TPair<FConcertObjectInStreamID, EConcertPutObjectErrorCode>& Error : ObjectsToPutSemanticErrors)
-	{
-		const FString ErrorCodeAsString = [&Error]()
-		{
-			switch (Error.Value)
-			{
-			case EConcertPutObjectErrorCode::UnresolvedStream: return TEXT("Unresolved stream");
-			case EConcertPutObjectErrorCode::MissingData: return TEXT("Missing data");
-			default: return TEXT("Unknown");
-			}
-		}();
-		OutputDevice.Logf(TEXT("Semantic error: %s - %s"), *Error.Key.ToString(), *ErrorCodeAsString);
-	}
-
-	for (const FGuid& Stream : FailedStreamCreation)
-	{
-		OutputDevice.Logf(TEXT("Failed to create stream: %s"), *Stream.ToString());
-	}
-}
+}
\ No newline at end of file
diff --git a/Engine/Plugins/Developer/Concert/ConcertSync/ConcertSyncCore/Source/ConcertSyncCore/Public/Replication/Data/ObjectIds.h b/Engine/Plugins/Developer/Concert/ConcertSync/ConcertSyncCore/Source/ConcertSyncCore/Public/Replication/Data/ObjectIds.h
--- a/Engine/Plugins/Developer/Concert/ConcertSync/ConcertSyncCore/Source/ConcertSyncCore/Public/Replication/Data/ObjectIds.h
+++ b/Engine/Plugins/Developer/Concert/ConcertSync/ConcertSyncCore/Source/ConcertSyncCore/Public/Replication/Data/ObjectIds.h
@@ -42,6 +42,7 @@ struct FConcertReplicatedObjectId : public FConcertObjectInStreamID
 	FGuid SenderEndpointId;
 	
 	FString ToString() const { return FString::Printf(TEXT("StreamId: %s, Object: %s, Sender: %s"), *StreamId.ToString(), *Object.ToString(), *SenderEndpointId.ToString()); }
+	FString ToString(const FString& ClientName) const { return FString::Printf(TEXT("StreamId: %s, Object: %s, Sender: %s"), *StreamId.ToString(), *Object.ToString(), *ClientName); }
 		
 	friend bool operator==(const FConcertReplicatedObjectId& Left, const FConcertReplicatedObjectId& Right)
 	{
diff --git a/Engine/Plugins/Developer/Concert/ConcertSync/ConcertSyncCore/Source/ConcertSyncCore/Public/Replication/Messages/ChangeStream.h b/Engine/Plugins/Developer/Concert/ConcertSync/ConcertSyncCore/Source/ConcertSyncCore/Public/Replication/Messages/ChangeStream.h
--- a/Engine/Plugins/Developer/Concert/ConcertSync/ConcertSyncCore/Source/ConcertSyncCore/Public/Replication/Messages/ChangeStream.h
+++ b/Engine/Plugins/Developer/Concert/ConcertSync/ConcertSyncCore/Source/ConcertSyncCore/Public/Replication/Messages/ChangeStream.h
@@ -8,6 +8,7 @@
 #include "Replication/Data/ReplicationStream.h"
 
 #include "Misc/Optional.h"
+#include "Templates/Function.h"
 #include "ChangeStream.generated.h"
 
 class FOutputDevice;
@@ -312,7 +313,4 @@ struct FConcertReplication_ChangeStream_Response
 	bool IsFailure() const { return !IsSuccess(); }
 
 	bool WasObjectPutSuccessful(const FConcertObjectInStreamID& Object) const { return !AuthorityConflicts.Contains(Object) && !ObjectsToPutSemanticErrors.Contains(Object); }
-
-	/** If IsFailure(), logs the errors. */
-	CONCERTSYNCCORE_API void LogErrors(FOutputDevice& OutputDevice) const;
 };
\ No newline at end of file
diff --git a/Engine/Plugins/Developer/Concert/ConcertSync/ConcertSyncServer/Source/ConcertSyncServer/ConcertSyncServer.Build.cs b/Engine/Plugins/Developer/Concert/ConcertSync/ConcertSyncServer/Source/ConcertSyncServer/ConcertSyncServer.Build.cs
--- a/Engine/Plugins/Developer/Concert/ConcertSync/ConcertSyncServer/Source/ConcertSyncServer/ConcertSyncServer.Build.cs
+++ b/Engine/Plugins/Developer/Concert/ConcertSync/ConcertSyncServer/Source/ConcertSyncServer/ConcertSyncServer.Build.cs
@@ -28,6 +28,7 @@ namespace UnrealBuildTool.Rules
 				{
 					"ConcertTransport",
 					"Serialization",
+					"JsonUtilities"
 				}
 			);
 		}
diff --git a/Engine/Plugins/Developer/Concert/ConcertSync/ConcertSyncServer/Source/ConcertSyncServer/Private/Replication/ConcertServerReplicationManager.cpp b/Engine/Plugins/Developer/Concert/ConcertSync/ConcertSyncServer/Source/ConcertSyncServer/Private/Replication/ConcertServerReplicationManager.cpp
--- a/Engine/Plugins/Developer/Concert/ConcertSync/ConcertSyncServer/Source/ConcertSyncServer/Private/Replication/ConcertServerReplicationManager.cpp
+++ b/Engine/Plugins/Developer/Concert/ConcertSync/ConcertSyncServer/Source/ConcertSyncServer/Private/Replication/ConcertServerReplicationManager.cpp
@@ -13,9 +13,16 @@
 #include "Replication/Messages/Handshake.h"
 #include "Replication/Processing/ObjectReplicationCache.h"
 #include "Util/JoinRequestValidation.h"
+#include "Util/LogUtils.h"
 
 namespace UE::ConcertSyncServer::Replication
 {
+	TAutoConsoleVariable<bool> CVarLogAuthorityRequestsAndResponses(
+		TEXT("Concert.Replication.LogAuthorityRequestsAndResponses"),
+		false,
+		TEXT("Whether to log changes to authority.")
+		);
+	
 	FConcertServerReplicationManager::FConcertServerReplicationManager(TSharedRef<IConcertServerSession> InLiveSession)
 		: Session(MoveTemp(InLiveSession))
 		, ReplicationFormat(MakeShared<ConcertSyncCore::FFullObjectFormat>())
@@ -79,7 +86,9 @@ namespace UE::ConcertSyncServer::Replication
 		// Have a pair of logs before and after processing in case of potential disaster
 		UE_LOG(LogConcert, Log, TEXT("Received replication join request from endpoint %s"), *ConcertSessionContext.SourceEndpointId.ToString());
 		
+		LogNetworkMessage(CVarLogAuthorityRequestsAndResponses, Request, [&](){ return GetClientName(*Session, ConcertSessionContext.SourceEndpointId); });
 		const EConcertSessionResponseCode Result = InternalHandleJoinReplicationSessionRequest(ConcertSessionContext, Request, Response);
+		LogNetworkMessage(CVarLogAuthorityRequestsAndResponses, Response, [&](){ return GetClientName(*Session, ConcertSessionContext.SourceEndpointId); });
 		
 		UE_CLOG(Response.JoinErrorCode == EJoinReplicationErrorCode::Success, LogConcert, Log, TEXT("Accepted replication join request"));
 		UE_CLOG(Response.JoinErrorCode != EJoinReplicationErrorCode::Success, LogConcert, Log, TEXT("Rejected replication join request. %s: %s"), *ConcertSyncCore::Replication::LexJoinErrorCode(Response.JoinErrorCode), *Response.DetailedErrorMessage);
diff --git a/Engine/Plugins/Developer/Concert/ConcertSync/ConcertSyncServer/Source/ConcertSyncServer/Private/Replication/ConcertServerReplicationManager_ChangeStream.cpp b/Engine/Plugins/Developer/Concert/ConcertSync/ConcertSyncServer/Source/ConcertSyncServer/Private/Replication/ConcertServerReplicationManager_ChangeStream.cpp
--- a/Engine/Plugins/Developer/Concert/ConcertSync/ConcertSyncServer/Source/ConcertSyncServer/Private/Replication/ConcertServerReplicationManager_ChangeStream.cpp
+++ b/Engine/Plugins/Developer/Concert/ConcertSync/ConcertSyncServer/Source/ConcertSyncServer/Private/Replication/ConcertServerReplicationManager_ChangeStream.cpp
@@ -9,6 +9,7 @@
 #include "Replication/ConcertReplicationClient.h"
 #include "Replication/Messages/ChangeStream.h"
 #include "Replication/Messages/Handshake.h"
+#include "Util/LogUtils.h"
 
 namespace UE::ConcertSyncServer::Replication
 {
@@ -137,12 +138,19 @@ namespace UE::ConcertSyncServer::Replication
 			return OutResponse.IsSuccess();
 		}
 	}
-
+	
+	TAutoConsoleVariable<bool> CVarLogStreamRequestsAndResponses(
+		TEXT("Concert.Replication.LogStreamRequestsAndResponses"),
+		false,
+		TEXT("Whether to log changes to streams.")
+		);
+	
 	EConcertSessionResponseCode FConcertServerReplicationManager::HandleChangeStreamRequest(
 		const FConcertSessionContext& ConcertSessionContext,
 		const FConcertReplication_ChangeStream_Request& Request,
 		FConcertReplication_ChangeStream_Response& Response)
 	{
+		LogNetworkMessage(CVarLogStreamRequestsAndResponses, Request, [&](){ return GetClientName(*Session, ConcertSessionContext.SourceEndpointId); });
 		Response = {};
 		
 		const FGuid SendingClientId = ConcertSessionContext.SourceEndpointId;
@@ -165,6 +173,7 @@ namespace UE::ConcertSyncServer::Replication
 		}
 		
 		Response.ErrorCode = EReplicationResponseErrorCode::Handled;
+		LogNetworkMessage(CVarLogStreamRequestsAndResponses, Response, [&](){ return GetClientName(*Session, ConcertSessionContext.SourceEndpointId); });
 		return EConcertSessionResponseCode::Success;
 	}
 }
diff --git a/Engine/Plugins/Developer/Concert/ConcertSync/ConcertSyncServer/Source/ConcertSyncServer/Private/Replication/Util/LogUtils.h b/Engine/Plugins/Developer/Concert/ConcertSync/ConcertSyncServer/Source/ConcertSyncServer/Private/Replication/Util/LogUtils.h
new file mode 100644
--- /dev/null
+++ b/Engine/Plugins/Developer/Concert/ConcertSync/ConcertSyncServer/Source/ConcertSyncServer/Private/Replication/Util/LogUtils.h
@@ -0,0 +1,67 @@
+// Copyright Epic Games, Inc. All Rights Reserved.
+
+#pragma once
+
+#include "ConcertLogGlobal.h"
+#include "JsonObjectConverter.h"
+#include "HAL/IConsoleManager.h"
+#include "Logging/LogMacros.h"
+
+#include <type_traits>
+
+namespace UE::ConcertSyncServer
+{
+	enum class EClientNameFlags
+	{
+		None,
+		IncludeName = 1 << 0,
+		IncludeEndpointId = 1 << 1
+	};
+	ENUM_CLASS_FLAGS(EClientNameFlags);
+	
+	inline TOptional<FString> GetClientName(IConcertSession& Session, const FGuid& EndpointId, EClientNameFlags Flags = EClientNameFlags::IncludeName | EClientNameFlags::IncludeEndpointId)
+	{
+		FConcertSessionClientInfo Info;
+		if (Session.FindSessionClient(EndpointId, Info))
+		{
+			if (EnumHasAllFlags(Flags, EClientNameFlags::IncludeName | EClientNameFlags::IncludeEndpointId))
+			{
+				return FString::Printf(TEXT("%s(%s)"), *Info.ClientInfo.DisplayName, *EndpointId.ToString());
+			}
+			if (EnumHasAnyFlags(Flags, EClientNameFlags::IncludeName))
+			{
+				return Info.ClientInfo.DisplayName;
+			}
+			if (EnumHasAnyFlags(Flags, EClientNameFlags::IncludeEndpointId))
+			{
+				return EndpointId.ToString();
+			}
+		}
+		return {};
+	}
+	
+	template<typename TMessage, typename TGetClientName> requires std::is_invocable_r_v<TOptional<FString>, TGetClientName>
+	void LogNetworkMessage(const TMessage& Message, TGetClientName GetClientName = [](){ return TOptional<FString>{}; })
+	{
+		FString JsonString;
+		FJsonObjectConverter::UStructToJsonObjectString(TMessage::StaticStruct(), &Message, JsonString, 0, 0);
+		
+		if (const TOptional<FString> ClientName = GetClientName())
+		{
+			UE_LOG(LogConcert, Log, TEXT("%s from %s\n%s"), *TMessage::StaticStruct()->GetName(), **ClientName, *JsonString);
+		}
+		else
+		{
+			UE_LOG(LogConcert, Log, TEXT("%s\n%s"), *TMessage::StaticStruct()->GetName(), *JsonString);
+		}
+	}
+	
+	template<typename TMessage, typename TGetClientName> requires std::is_invocable_r_v<TOptional<FString>, TGetClientName>
+	void LogNetworkMessage(const TAutoConsoleVariable<bool>& ShouldLog, const TMessage& Message, TGetClientName GetClientName = [](){ return TOptional<FString>{}; })
+	{
+		if (ShouldLog.GetValueOnAnyThread())
+		{
+			LogNetworkMessage(Message, GetClientName);
+		}
+	}
+}
