Fix CIS caused by CVars having same names

diff --git a/Engine/Plugins/Developer/Concert/ConcertSync/ConcertSyncClient/Source/ConcertSyncClient/Private/Replication/Manager/ReplicationManagerState_Connected.cpp b/Engine/Plugins/Developer/Concert/ConcertSync/ConcertSyncClient/Source/ConcertSyncClient/Private/Replication/Manager/ReplicationManagerState_Connected.cpp
--- a/Engine/Plugins/Developer/Concert/ConcertSync/ConcertSyncClient/Source/ConcertSyncClient/Private/Replication/Manager/ReplicationManagerState_Connected.cpp
+++ b/Engine/Plugins/Developer/Concert/ConcertSync/ConcertSyncClient/Source/ConcertSyncClient/Private/Replication/Manager/ReplicationManagerState_Connected.cpp
@@ -43,13 +43,13 @@ namespace UE::ConcertSyncClient::Replication
 		TEXT("Whether the client should pretend that authority change requests were rejected.")
 		);
 
-	TAutoConsoleVariable<bool> CVarLogStreamRequestsAndResponses(
-		TEXT("Concert.Replication.LogStreamRequestsAndResponses"),
+	TAutoConsoleVariable<bool> CVarLogStreamRequestsAndResponsesOnClient(
+		TEXT("Concert.Replication.LogStreamRequestsAndResponsesOnClient"),
 		false,
 		TEXT("Whether to log changes to streams.")
 		);
-	TAutoConsoleVariable<bool> CVarLogAuthorityRequestsAndResponses(
-		TEXT("Concert.Replication.LogStreamRequestsAndResponses"),
+	TAutoConsoleVariable<bool> CVarLogAuthorityRequestsAndResponsesOnClient(
+		TEXT("Concert.Replication.LogAuthorityRequestsAndResponsesOnClient"),
 		false,
 		TEXT("Whether to log changes to authority.")
 		);
@@ -146,11 +146,11 @@ namespace UE::ConcertSyncClient::Replication
 		// At that point, it will log errors for receiving replication data from a client without authority.
 		HandleReleasingReplicatedObjects(Args);
 
-		Private::LogNetworkMessage(CVarLogAuthorityRequestsAndResponses, Args);
+		Private::LogNetworkMessage(CVarLogAuthorityRequestsAndResponsesOnClient, Args);
 		return LiveSession->SendCustomRequest<FConcertReplication_ChangeAuthority_Request, FConcertReplication_ChangeAuthority_Response>(Args, LiveSession->GetSessionServerEndpointId())
 			.Next([WeakThis = TWeakPtr<FReplicationManagerState_Connected>(SharedThis(this)), Args](FConcertReplication_ChangeAuthority_Response&& Response) mutable
 			{
-				Private::LogNetworkMessage(CVarLogAuthorityRequestsAndResponses, Response);
+				Private::LogNetworkMessage(CVarLogAuthorityRequestsAndResponsesOnClient, Response);
 				
 				if (const TSharedPtr<FReplicationManagerState_Connected> ThisPin = WeakThis.Pin()
 					; ThisPin && Response.ErrorCode == EReplicationResponseErrorCode::Handled)
@@ -198,11 +198,11 @@ namespace UE::ConcertSyncClient::Replication
 		// At that point, it will log errors for receiving replication data from a client without authority.
 		HandleRemovingReplicatedObjects(Args);
 		
-		Private::LogNetworkMessage(CVarLogStreamRequestsAndResponses, Args);
+		Private::LogNetworkMessage(CVarLogStreamRequestsAndResponsesOnClient, Args);
 		return LiveSession->SendCustomRequest<FConcertReplication_ChangeStream_Request, FConcertReplication_ChangeStream_Response>(Args, LiveSession->GetSessionServerEndpointId())
 			.Next([WeakThis = TWeakPtr<FReplicationManagerState_Connected>(SharedThis(this)), Args](FConcertReplication_ChangeStream_Response&& Response)
 			{
-				Private::LogNetworkMessage(CVarLogStreamRequestsAndResponses, Response);
+				Private::LogNetworkMessage(CVarLogStreamRequestsAndResponsesOnClient, Response);
 
 				const TSharedPtr<FReplicationManagerState_Connected> ThisPin = WeakThis.Pin();
 				if (ThisPin && Response.IsSuccess())
diff --git a/Engine/Plugins/Developer/Concert/ConcertSync/ConcertSyncServer/Source/ConcertSyncServer/Private/Replication/ConcertServerReplicationManager.cpp b/Engine/Plugins/Developer/Concert/ConcertSync/ConcertSyncServer/Source/ConcertSyncServer/Private/Replication/ConcertServerReplicationManager.cpp
--- a/Engine/Plugins/Developer/Concert/ConcertSync/ConcertSyncServer/Source/ConcertSyncServer/Private/Replication/ConcertServerReplicationManager.cpp
+++ b/Engine/Plugins/Developer/Concert/ConcertSync/ConcertSyncServer/Source/ConcertSyncServer/Private/Replication/ConcertServerReplicationManager.cpp
@@ -17,8 +17,8 @@
 
 namespace UE::ConcertSyncServer::Replication
 {
-	TAutoConsoleVariable<bool> CVarLogAuthorityRequestsAndResponses(
-		TEXT("Concert.Replication.LogAuthorityRequestsAndResponses"),
+	TAutoConsoleVariable<bool> CVarLogAuthorityRequestsAndResponsesOnServer(
+		TEXT("Concert.Replication.LogAuthorityRequestsAndResponsesOnServer"),
 		false,
 		TEXT("Whether to log changes to authority.")
 		);
@@ -86,9 +86,9 @@ namespace UE::ConcertSyncServer::Replication
 		// Have a pair of logs before and after processing in case of potential disaster
 		UE_LOG(LogConcert, Log, TEXT("Received replication join request from endpoint %s"), *ConcertSessionContext.SourceEndpointId.ToString());
 		
-		LogNetworkMessage(CVarLogAuthorityRequestsAndResponses, Request, [&](){ return GetClientName(*Session, ConcertSessionContext.SourceEndpointId); });
+		LogNetworkMessage(CVarLogAuthorityRequestsAndResponsesOnServer, Request, [&](){ return GetClientName(*Session, ConcertSessionContext.SourceEndpointId); });
 		const EConcertSessionResponseCode Result = InternalHandleJoinReplicationSessionRequest(ConcertSessionContext, Request, Response);
-		LogNetworkMessage(CVarLogAuthorityRequestsAndResponses, Response, [&](){ return GetClientName(*Session, ConcertSessionContext.SourceEndpointId); });
+		LogNetworkMessage(CVarLogAuthorityRequestsAndResponsesOnServer, Response, [&](){ return GetClientName(*Session, ConcertSessionContext.SourceEndpointId); });
 		
 		UE_CLOG(Response.JoinErrorCode == EJoinReplicationErrorCode::Success, LogConcert, Log, TEXT("Accepted replication join request"));
 		UE_CLOG(Response.JoinErrorCode != EJoinReplicationErrorCode::Success, LogConcert, Log, TEXT("Rejected replication join request. %s: %s"), *ConcertSyncCore::Replication::LexJoinErrorCode(Response.JoinErrorCode), *Response.DetailedErrorMessage);
diff --git a/Engine/Plugins/Developer/Concert/ConcertSync/ConcertSyncServer/Source/ConcertSyncServer/Private/Replication/ConcertServerReplicationManager_ChangeStream.cpp b/Engine/Plugins/Developer/Concert/ConcertSync/ConcertSyncServer/Source/ConcertSyncServer/Private/Replication/ConcertServerReplicationManager_ChangeStream.cpp
--- a/Engine/Plugins/Developer/Concert/ConcertSync/ConcertSyncServer/Source/ConcertSyncServer/Private/Replication/ConcertServerReplicationManager_ChangeStream.cpp
+++ b/Engine/Plugins/Developer/Concert/ConcertSync/ConcertSyncServer/Source/ConcertSyncServer/Private/Replication/ConcertServerReplicationManager_ChangeStream.cpp
@@ -139,8 +139,8 @@ namespace UE::ConcertSyncServer::Replication
 		}
 	}
 	
-	TAutoConsoleVariable<bool> CVarLogStreamRequestsAndResponses(
-		TEXT("Concert.Replication.LogStreamRequestsAndResponses"),
+	TAutoConsoleVariable<bool> CVarLogStreamRequestsAndResponsesOnServer(
+		TEXT("Concert.Replication.LogStreamRequestsAndResponsesOnServer"),
 		false,
 		TEXT("Whether to log changes to streams.")
 		);
@@ -150,7 +150,7 @@ namespace UE::ConcertSyncServer::Replication
 		const FConcertReplication_ChangeStream_Request& Request,
 		FConcertReplication_ChangeStream_Response& Response)
 	{
-		LogNetworkMessage(CVarLogStreamRequestsAndResponses, Request, [&](){ return GetClientName(*Session, ConcertSessionContext.SourceEndpointId); });
+		LogNetworkMessage(CVarLogStreamRequestsAndResponsesOnServer, Request, [&](){ return GetClientName(*Session, ConcertSessionContext.SourceEndpointId); });
 		Response = {};
 		
 		const FGuid SendingClientId = ConcertSessionContext.SourceEndpointId;
@@ -173,7 +173,7 @@ namespace UE::ConcertSyncServer::Replication
 		}
 		
 		Response.ErrorCode = EReplicationResponseErrorCode::Handled;
-		LogNetworkMessage(CVarLogStreamRequestsAndResponses, Response, [&](){ return GetClientName(*Session, ConcertSessionContext.SourceEndpointId); });
+		LogNetworkMessage(CVarLogStreamRequestsAndResponsesOnServer, Response, [&](){ return GetClientName(*Session, ConcertSessionContext.SourceEndpointId); });
 		return EConcertSessionResponseCode::Success;
 	}
 }
