Handle clients joining and leaving while replication is occuring

diff --git a/Engine/Plugins/Developer/Concert/ConcertSync/ConcertSyncClient/Source/ConcertSyncClient/Private/Replication/ConcertClientReplicationBridge.cpp b/Engine/Plugins/Developer/Concert/ConcertSync/ConcertSyncClient/Source/ConcertSyncClient/Private/Replication/ConcertClientReplicationBridge.cpp
--- a/Engine/Plugins/Developer/Concert/ConcertSync/ConcertSyncClient/Source/ConcertSyncClient/Private/Replication/ConcertClientReplicationBridge.cpp
+++ b/Engine/Plugins/Developer/Concert/ConcertSync/ConcertSyncClient/Source/ConcertSyncClient/Private/Replication/ConcertClientReplicationBridge.cpp
@@ -4,18 +4,34 @@
 
 #include "ConcertLogGlobal.h"
 
-#include "AssetRegistry/IAssetRegistry.h"
-#include "AssetRegistry/AssetData.h"
+#include "Engine/Engine.h"
 #include "Engine/World.h"
-
-class IAssetRegistry;
+#include "GameFramework/Actor.h"
+#include "Misc/CoreDelegates.h"
 
 namespace UE::ConcertSyncClient::Replication
 {
 	FConcertClientReplicationBridge::FConcertClientReplicationBridge()
 	{
-		FWorldDelegates::OnPostWorldInitialization.AddRaw(this, &FConcertClientReplicationBridge::OnPostWorldInitialization);
-		FWorldDelegates::OnPostWorldCleanup.AddRaw(this, &FConcertClientReplicationBridge::OnPostWorldCleanup);
+		// This constructor is called by StartupModule: at that time GEngine is not inited yet.
+		// The constructor is not supposed "know" when it is called though so handle it here.
+		if (!GEngine)
+		{
+			FCoreDelegates::OnPostEngineInit.AddRaw(this, &FConcertClientReplicationBridge::SetupEngineDelegates);
+		}
+		else
+		{
+			SetupEngineDelegates();
+		}
+
+		OnObjectDiscoveredDelegate.AddLambda([](UObject& Object)
+		{
+			UE_LOG(LogConcert, Verbose, TEXT("Discovered object %s for replication"), *Object.GetPathName());
+		});
+		OnObjectRemovedDelegate.AddLambda([](const FSoftObjectPath& ObjectPath)
+		{
+			UE_LOG(LogConcert, Verbose, TEXT("Removed object %s from replication"), *ObjectPath.ToString());
+		});
 	}
 
 	FConcertClientReplicationBridge::~FConcertClientReplicationBridge()
@@ -30,7 +46,7 @@ namespace UE::ConcertSyncClient::Replication
 		
 		for (const FSoftObjectPath& TrackedObjectPath : InTrackedObjects)
 		{
-			FTrackedObject& TrackedObject = TrackedObjects.FindOrAdd(TrackedObjectPath);
+			FTrackedObjectInfo& TrackedObject = TrackedObjects.FindOrAdd(TrackedObjectPath);
 			++TrackedObject.TrackCounter;
 			if (!TrackedObject.ResolvedObject.IsValid())
 			{
@@ -54,7 +70,7 @@ namespace UE::ConcertSyncClient::Replication
 	{
 		for (const FSoftObjectPath& TrackedObjectPath : InTrackedObjects)
 		{
-			FTrackedObject* TrackedObject = TrackedObjects.Find(TrackedObjectPath);
+			FTrackedObjectInfo* TrackedObject = TrackedObjects.Find(TrackedObjectPath);
 			if (!TrackedObject)
 			{
 				UE_LOG(LogConcert, Warning, TEXT("Object %s was requested to not be tracked but was it already wasn't."), *TrackedObjectPath.ToString());
@@ -80,47 +96,193 @@ namespace UE::ConcertSyncClient::Replication
 
 	UObject* FConcertClientReplicationBridge::FindObjectIfAvailable(const FSoftObjectPath& Path)
 	{
-		IAssetRegistry* AssetRegistry = IAssetRegistry::Get();
-		if (!AssetRegistry)
-		{
-			return nullptr;
-		}
+		// This should be an object in a UWorld. This will resolve if the world is opened.
+		UObject* Object = Path.ResolveObject();
+		UE_CLOG(Object && !Object->IsInA(UWorld::StaticClass()), LogConcert, Warning, TEXT("Object %s is not in any UWorld. The replication system was designed to work with objects in UWorlds."), *Path.ToString());
+		return Object;
+	}
 
-		constexpr bool bSkipInMemory = true;
-		const FAssetData AssetData = AssetRegistry->GetAssetByObjectPath(Path, bSkipInMemory);
-		const bool bIsAsset = AssetData.IsValid() || Path.IsAsset();
-		if (!ensure(!bIsAsset))
+	void FConcertClientReplicationBridge::SetupEngineDelegates()
+	{
+		if (ensure(GEngine))
 		{
-			UE_LOG(LogConcert, Error, TEXT("Assets are not supported for replication in this version (%s)"), *Path.ToString());
-			return nullptr;
+			GEngine->OnWorldAdded().AddRaw(this, &FConcertClientReplicationBridge::SetupWorld);
+			GEngine->OnWorldDestroyed().AddRaw(this, &FConcertClientReplicationBridge::CleanupWorld);
+			// Important: In a Multi-User session, actors might get added only after the world is loaded. Transactions are replayed at the end of tick.
+			// Replaying calls SpawnActor so the below events will trigger.
+			GEngine->OnLevelActorAdded().AddRaw(this, &FConcertClientReplicationBridge::DiscoverTrackedObjectsInActor);
+			GEngine->OnLevelActorDeleted().AddRaw(this, &FConcertClientReplicationBridge::RemoveDiscoveredObjectsInActor);
 		}
-		
-		// This should be an object in a UWorld. This will resolve if the world is opened.
-		return Path.ResolveObject();
 	}
 
-	void FConcertClientReplicationBridge::OnPostWorldInitialization(UWorld* World, FWorldInitializationValues WorldInitializationValues)
+	void FConcertClientReplicationBridge::SetupWorld(UWorld* World)
 	{
 		// We'll allow PIE even though it is currently not explicitly supported; the primary use case is editor or game worlds (if someone tries it in packaged).
-		if (!ensure(World) || World->IsPreviewWorld())
+		if (World->IsPreviewWorld())
 		{
 			return;
 		}
 		
 		LoadedWorlds.Add(World);
-
-		// TODO: Discover objects
+		DiscoverTrackedObjectsInWorld(*World);
 	}
 
-	void FConcertClientReplicationBridge::OnPostWorldCleanup(UWorld* World, bool bSessionEnded, bool bCleanupResources)
+	void FConcertClientReplicationBridge::CleanupWorld(UWorld* World)
 	{
 		const int32 NumRemoved = LoadedWorlds.Remove(World);
 		bool bRemovedAnything = NumRemoved > 0;
-		if (!bRemovedAnything)
+		if (bRemovedAnything)
 		{
-			return;
+			RemoveDiscoveredObjectsInWorld(*World);
+		}
+	}
+
+	void FConcertClientReplicationBridge::DiscoverTrackedObjectsInWorld(UWorld& World)
+	{
+		UE_LOG(LogConcert, Verbose, TEXT("Discovering objects in world %s for replication"), *World.GetPathName());
+		// Defer calling the delegate in case it causes a world unload or changes the tracking state.
+		TArray<TWeakObjectPtr<UObject>> DeferredDiscoveredObjects;
+		
+		for (TPair<FSoftObjectPath, FTrackedObjectInfo>& Pair : TrackedObjects)
+		{
+			FTrackedObjectInfo& ObjectInfo = Pair.Value;
+			if (ObjectInfo.ResolvedObject.IsValid())
+			{
+				continue;
+			}
+
+			const FSoftObjectPath& ObjectPath = Pair.Key;
+			UObject* Object = ObjectPath.ResolveObject();
+			// Only objects from the loaded world should be considered. ResolveObject will still resolve to objects from the old world.
+			if (Object && Object->IsIn(&World))
+			{
+				ObjectInfo.ResolvedObject = Object;
+				DeferredDiscoveredObjects.Emplace(Object);
+			}
 		}
 
-		// TODO: Hide objects
+		MarkAsDiscovered(DeferredDiscoveredObjects);
+	}
+
+	void FConcertClientReplicationBridge::RemoveDiscoveredObjectsInWorld(UWorld& World)
+	{
+		UE_LOG(LogConcert, Verbose, TEXT("Removing objects from world %s for replication"), *World.GetPathName());
+		// Defer calling the delegate in case it causes a world load or changes the tracking state.
+		TArray<FSoftObjectPath> DeferredHiddenObjects;
+		
+		for (TPair<FSoftObjectPath, FTrackedObjectInfo>& Pair : TrackedObjects)
+		{
+			TWeakObjectPtr<UObject>& WeakObject = Pair.Value.ResolvedObject;
+			if (WeakObject.IsStale())
+			{
+				WeakObject.Reset();
+				DeferredHiddenObjects.Add(Pair.Key);
+			}
+			// Not sure whether the weak ptr would return IsStale in a OnPostWorldCleanup callback... this makes sure.
+			else if (UObject* Object = Pair.Key.ResolveObject()
+				; Object && Object->IsIn(&World))
+			{
+				WeakObject.Reset();
+				DeferredHiddenObjects.Add(Pair.Key);
+			}
+		}
+
+		MarkAsUndiscovered(DeferredHiddenObjects);
+	}
+
+	void FConcertClientReplicationBridge::DiscoverTrackedObjectsIn(UObject* AnalysedObject)
+	{
+		TArray<TWeakObjectPtr<UObject>> FoundObjects;
+		auto AnalyseObject = [this, &FoundObjects](UObject* Object)
+		{
+			FTrackedObjectInfo* ObjectInfo = TrackedObjects.Find(Object);
+			if (!ObjectInfo)
+			{
+				// Object is not being tracked
+				return;
+			}
+
+			if (!ObjectInfo->ResolvedObject.IsValid())
+			{
+				ObjectInfo->ResolvedObject = Object;
+				FoundObjects.Emplace(Object);
+			}
+		};
+		
+		AnalyseObject(AnalysedObject);
+		constexpr bool bSearchedNested = true;
+		ForEachObjectWithOuter(AnalysedObject, [&AnalyseObject](UObject* Subobject)
+		{
+			AnalyseObject(Subobject);
+		}, bSearchedNested);
+
+		MarkAsDiscovered(FoundObjects);
+	}
+
+	void FConcertClientReplicationBridge::RemoveDiscoveredObjectsIn(UObject* Object)
+	{
+		TArray<FSoftObjectPath> ObjectsToRemove;
+		auto AnalyseObject = [this, &ObjectsToRemove](UObject* Object)
+		{
+			const FSoftObjectPath ObjectPath = Object;
+			FTrackedObjectInfo* ObjectInfo = TrackedObjects.Find(ObjectPath);
+			if (!ObjectInfo)
+			{
+				// Object is not being tracked
+				return;
+			}
+
+			if (ObjectInfo->ResolvedObject == Object)
+			{
+				ObjectInfo->ResolvedObject.Reset();
+				ObjectsToRemove.Emplace(ObjectPath);
+			}
+		};
+		
+		AnalyseObject(Object);
+		constexpr bool bSearchedNested = true;
+		ForEachObjectWithOuter(Object, [&AnalyseObject](UObject* Subobject)
+		{
+			AnalyseObject(Subobject);
+		}, bSearchedNested);
+
+		MarkAsUndiscovered(ObjectsToRemove);
+	}
+
+	void FConcertClientReplicationBridge::DiscoverTrackedObjectsInActor(AActor* Actor)
+	{
+		DiscoverTrackedObjectsIn(Actor);
+	}
+
+	void FConcertClientReplicationBridge::RemoveDiscoveredObjectsInActor(AActor* Actor)
+	{
+		RemoveDiscoveredObjectsIn(Actor);
+	}
+
+	void FConcertClientReplicationBridge::MarkAsDiscovered(TConstArrayView<TWeakObjectPtr<UObject>> DeferredDiscoveredObjects)
+	{
+		for (const TWeakObjectPtr<UObject>& WeakObject : DeferredDiscoveredObjects)
+		{
+			// The point of checking this again is that a previous OnObjectDiscoveredDelegate call may have caused a world unload or stopped tracking something.
+			// If a world unload was triggered, then OnPostWorldCleanup will have cleaned FTrackedObjectInfo::ResolvedObject (in call recursive to this one).
+			if (WeakObject.IsValid() && TrackedObjects.Contains(WeakObject.Get()))
+			{
+				OnObjectDiscoveredDelegate.Broadcast(*WeakObject);
+			}
+		}
+	}
+	void FConcertClientReplicationBridge::MarkAsUndiscovered(TConstArrayView<FSoftObjectPath> DeferredHiddenObjects)
+	{
+		for (const FSoftObjectPath& HiddenObjectPath : DeferredHiddenObjects)
+		{
+			const UObject* Resolved = HiddenObjectPath.ResolveObject();
+			
+			// The point of checking this again is that a previous OnObjectRemovedDelegate call may have caused a world load or stopped tracking something.
+			// If a world load was triggered, then OnPostWorldInitialization will have set FTrackedObjectInfo::ResolvedObject (in call recursive to this one).
+			if (TrackedObjects.Contains(HiddenObjectPath))
+			{
+				OnObjectRemovedDelegate.Broadcast(HiddenObjectPath);
+			}
+		}
 	}
 }
diff --git a/Engine/Plugins/Developer/Concert/ConcertSync/ConcertSyncClient/Source/ConcertSyncClient/Private/Replication/ConcertClientReplicationBridge.h b/Engine/Plugins/Developer/Concert/ConcertSync/ConcertSyncClient/Source/ConcertSyncClient/Private/Replication/ConcertClientReplicationBridge.h
--- a/Engine/Plugins/Developer/Concert/ConcertSync/ConcertSyncClient/Source/ConcertSyncClient/Private/Replication/ConcertClientReplicationBridge.h
+++ b/Engine/Plugins/Developer/Concert/ConcertSync/ConcertSyncClient/Source/ConcertSyncClient/Private/Replication/ConcertClientReplicationBridge.h
@@ -2,13 +2,15 @@
 
 #pragma once
 
+#include "Replication/IConcertClientReplicationBridge.h"
+
 #include "Containers/Map.h"
 #include "Containers/Set.h"
 #include "UObject/SoftObjectPath.h"
 #include "UObject/WeakObjectPtr.h"
 #include "Templates/SharedPointer.h"
-#include "Replication/IConcertClientReplicationBridge.h"
 
+class AActor;
 class UObject;
 class UWorld;
 struct FWorldInitializationValues;
@@ -28,12 +30,12 @@ namespace UE::ConcertSyncClient::Replication
 		virtual bool IsObjectAvailable(const FSoftObjectPath& Path) override;
 		virtual UObject* FindObjectIfAvailable(const FSoftObjectPath& Path) override;
 		virtual FConcertClientReplicationBridgeObjectEvent& OnObjectDiscovered() override { return OnObjectDiscoveredDelegate; }
-		virtual FConcertClientReplicationBridgeObjectEvent& OnObjectHidden() override { return OnObjectRemovedDelegate; }
+		virtual FConcertClientReplicationBridgeObjectPathEvent& OnObjectHidden() override { return OnObjectRemovedDelegate; }
 		//~ End IConcertClientReplicationBridge Interface
 
 	private:
 
-		struct FTrackedObject
+		struct FTrackedObjectInfo
 		{
 			/** Tracks the number of PushTrackedObjects calls. This entry is removed upon reaching 0. */
 			int32 TrackCounter = 0;
@@ -42,17 +44,32 @@ namespace UE::ConcertSyncClient::Replication
 		};
 
 		/** The objects that wish to be tracked. */
-		TMap<FSoftObjectPath, FTrackedObject> TrackedObjects;
+		TMap<FSoftObjectPath, FTrackedObjectInfo> TrackedObjects;
 
 		/** The set of worlds currently open on this client. */
 		TSet<TWeakObjectPtr<UWorld>> LoadedWorlds;
 	
 		FConcertClientReplicationBridgeObjectEvent OnObjectDiscoveredDelegate;
-		FConcertClientReplicationBridgeObjectEvent OnObjectRemovedDelegate;
+		FConcertClientReplicationBridgeObjectPathEvent OnObjectRemovedDelegate;
+
+		// Setting up delegates for worlds
+		void SetupEngineDelegates();
+		void SetupWorld(UWorld* World);
+		void CleanupWorld(UWorld* World);
+
+		// Searching everything in world
+		void DiscoverTrackedObjectsInWorld(UWorld& World);
+		void RemoveDiscoveredObjectsInWorld(UWorld& World);
+
+		// Searching a specific object
+		void DiscoverTrackedObjectsIn(UObject* AnalysedObject);
+		void RemoveDiscoveredObjectsIn(UObject* Object);
+		void DiscoverTrackedObjectsInActor(AActor* Actor);
+		void RemoveDiscoveredObjectsInActor(AActor* Actor);
 
-		// Callbacks into the engine
-		void OnPostWorldInitialization(UWorld* World, FWorldInitializationValues WorldInitializationValues);
-		void OnPostWorldCleanup(UWorld* World, bool bSessionEnded, bool bCleanupResources);
+		// Util for formally marking objects as discovered
+		void MarkAsDiscovered(TConstArrayView<TWeakObjectPtr<UObject>> DeferredDiscoveredObjects);
+		void MarkAsUndiscovered(TConstArrayView<FSoftObjectPath> DeferredHiddenObjects);
 	};
 }
 
diff --git a/Engine/Plugins/Developer/Concert/ConcertSync/ConcertSyncClient/Source/ConcertSyncClient/Private/Replication/Processing/ClientReplicationDataCollector.cpp b/Engine/Plugins/Developer/Concert/ConcertSync/ConcertSyncClient/Source/ConcertSyncClient/Private/Replication/Processing/ClientReplicationDataCollector.cpp
--- a/Engine/Plugins/Developer/Concert/ConcertSync/ConcertSyncClient/Source/ConcertSyncClient/Private/Replication/Processing/ClientReplicationDataCollector.cpp
+++ b/Engine/Plugins/Developer/Concert/ConcertSync/ConcertSyncClient/Source/ConcertSyncClient/Private/Replication/Processing/ClientReplicationDataCollector.cpp
@@ -226,9 +226,9 @@ namespace UE::ConcertSyncClient::Replication
 		}
 	}
 
-	void FClientReplicationDataCollector::StopTrackingObject(UObject& Object)
+	void FClientReplicationDataCollector::StopTrackingObject(const FSoftObjectPath& ObjectPath)
 	{
-		if (TArray<FObjectInfo>* ObjectInfos = ObjectsToReplicate.Find(&Object))
+		if (TArray<FObjectInfo>* ObjectInfos = ObjectsToReplicate.Find(ObjectPath))
 		{
 			for (FObjectInfo& ObjectInfo : *ObjectInfos)
 			{
diff --git a/Engine/Plugins/Developer/Concert/ConcertSync/ConcertSyncClient/Source/ConcertSyncClient/Private/Replication/Processing/ClientReplicationDataCollector.h b/Engine/Plugins/Developer/Concert/ConcertSync/ConcertSyncClient/Source/ConcertSyncClient/Private/Replication/Processing/ClientReplicationDataCollector.h
--- a/Engine/Plugins/Developer/Concert/ConcertSync/ConcertSyncClient/Source/ConcertSyncClient/Private/Replication/Processing/ClientReplicationDataCollector.h
+++ b/Engine/Plugins/Developer/Concert/ConcertSync/ConcertSyncClient/Source/ConcertSyncClient/Private/Replication/Processing/ClientReplicationDataCollector.h
@@ -110,6 +110,6 @@ namespace UE::ConcertSyncClient::Replication
 
 		// Handle events from bridge
 		void StartTrackingObject(UObject& Object);
-		void StopTrackingObject(UObject& Object);
+		void StopTrackingObject(const FSoftObjectPath& ObjectPath);
 	};
 }
diff --git a/Engine/Plugins/Developer/Concert/ConcertSync/ConcertSyncClient/Source/ConcertSyncClient/Public/Replication/IConcertClientReplicationBridge.h b/Engine/Plugins/Developer/Concert/ConcertSync/ConcertSyncClient/Source/ConcertSyncClient/Public/Replication/IConcertClientReplicationBridge.h
--- a/Engine/Plugins/Developer/Concert/ConcertSync/ConcertSyncClient/Source/ConcertSyncClient/Public/Replication/IConcertClientReplicationBridge.h
+++ b/Engine/Plugins/Developer/Concert/ConcertSync/ConcertSyncClient/Source/ConcertSyncClient/Public/Replication/IConcertClientReplicationBridge.h
@@ -11,6 +11,7 @@ struct FConcertReplication_ObjectReplicationEvent;
 struct FSoftObjectPath;
 
 DECLARE_MULTICAST_DELEGATE_OneParam(FConcertClientReplicationBridgeObjectEvent, UObject&);
+DECLARE_MULTICAST_DELEGATE_OneParam(FConcertClientReplicationBridgeObjectPathEvent, const FSoftObjectPath&);
 
 /**
  * A session-independent object bridging Unreal Engine with Concert.
@@ -77,5 +78,5 @@ public:
 	/** Event for when a replicatable object becomes available, e.g. because its level is loaded. */
 	virtual FConcertClientReplicationBridgeObjectEvent& OnObjectDiscovered() = 0;
 	/** Event for when an object that was previously discovered becomes unavailable, e.g. because it is unloaded. */
-	virtual FConcertClientReplicationBridgeObjectEvent& OnObjectHidden() = 0;
+	virtual FConcertClientReplicationBridgeObjectPathEvent& OnObjectHidden() = 0;
 };
diff --git a/Engine/Plugins/Developer/Concert/ConcertSync/ConcertSyncTest/Source/ConcertSyncTest/Private/Replication/Util/ConcertClientReplicationBridgeMock.h b/Engine/Plugins/Developer/Concert/ConcertSync/ConcertSyncTest/Source/ConcertSyncTest/Private/Replication/Util/ConcertClientReplicationBridgeMock.h
--- a/Engine/Plugins/Developer/Concert/ConcertSync/ConcertSyncTest/Source/ConcertSyncTest/Private/Replication/Util/ConcertClientReplicationBridgeMock.h
+++ b/Engine/Plugins/Developer/Concert/ConcertSync/ConcertSyncTest/Source/ConcertSyncTest/Private/Replication/Util/ConcertClientReplicationBridgeMock.h
@@ -16,7 +16,7 @@ namespace UE::ConcertSyncTests::Replication
 		TArray<UObject*> AvailableObjects;
 		
 		FConcertClientReplicationBridgeObjectEvent OnObjectDiscoveredDelegate;
-		FConcertClientReplicationBridgeObjectEvent OnObjectRemovedDelegate;
+		FConcertClientReplicationBridgeObjectPathEvent OnObjectRemovedDelegate;
 
 		void InjectAvailableObject(UObject& Object)
 		{
@@ -63,6 +63,6 @@ namespace UE::ConcertSyncTests::Replication
 		}
 
 		virtual FConcertClientReplicationBridgeObjectEvent& OnObjectDiscovered() override { return OnObjectDiscoveredDelegate; }
-		virtual FConcertClientReplicationBridgeObjectEvent& OnObjectHidden() override { return OnObjectRemovedDelegate; }
+		virtual FConcertClientReplicationBridgeObjectPathEvent& OnObjectHidden() override { return OnObjectRemovedDelegate; }
 	};
 }
