Fix missing deregistration of GEngine delegates

diff --git a/Engine/Plugins/Developer/Concert/ConcertSync/ConcertSyncClient/Source/ConcertSyncClient/Private/Replication/ConcertClientReplicationBridge.cpp b/Engine/Plugins/Developer/Concert/ConcertSync/ConcertSyncClient/Source/ConcertSyncClient/Private/Replication/ConcertClientReplicationBridge.cpp
--- a/Engine/Plugins/Developer/Concert/ConcertSync/ConcertSyncClient/Source/ConcertSyncClient/Private/Replication/ConcertClientReplicationBridge.cpp
+++ b/Engine/Plugins/Developer/Concert/ConcertSync/ConcertSyncClient/Source/ConcertSyncClient/Private/Replication/ConcertClientReplicationBridge.cpp
@@ -36,8 +36,13 @@ namespace UE::ConcertSyncClient::Replication
 
 	FConcertClientReplicationBridge::~FConcertClientReplicationBridge()
 	{
-		FWorldDelegates::OnPostWorldInitialization.RemoveAll(this);
-		FWorldDelegates::OnPostWorldCleanup.RemoveAll(this);
+		if (GEngine)
+		{
+			GEngine->OnWorldAdded().RemoveAll(this);
+			GEngine->OnWorldDestroyed().RemoveAll(this);
+			GEngine->OnLevelActorAdded().RemoveAll(this);
+			GEngine->OnLevelActorDeleted().RemoveAll(this);
+		}
 	}
 
 	void FConcertClientReplicationBridge::PushTrackedObjects(TArrayView<const FSoftObjectPath> InTrackedObjects)
