MU Replication: when applying received replicated data, fix resetting the transform back to the value it had just before the replicated data was applied.
- This was an "obvious" mistake when replicating component transforms and should not have worked in the first place, but did in some setups. It became apparent on stage.

diff --git a/Engine/Plugins/Developer/Concert/ConcertSync/ConcertSyncClient/Source/ConcertSyncClient/Private/Replication/Processing/ObjectReplicationApplierProcessor.cpp b/Engine/Plugins/Developer/Concert/ConcertSync/ConcertSyncClient/Source/ConcertSyncClient/Private/Replication/Processing/ObjectReplicationApplierProcessor.cpp
--- a/Engine/Plugins/Developer/Concert/ConcertSync/ConcertSyncClient/Source/ConcertSyncClient/Private/Replication/Processing/ObjectReplicationApplierProcessor.cpp
+++ b/Engine/Plugins/Developer/Concert/ConcertSync/ConcertSyncClient/Source/ConcertSyncClient/Private/Replication/Processing/ObjectReplicationApplierProcessor.cpp
@@ -37,9 +37,7 @@ namespace UE::ConcertSyncClient::Replication
 			// TODO DP UE-193659: This is very hacky and leaves performance on the table... this is in case ApplyReplicationEvent updates the transform
 			if (USceneComponent* SceneComponent = Cast<USceneComponent>(Object))
 			{
-				FTransform TransformBefore = SceneComponent->GetComponentTransform();
 				ReplicationFormat->ApplyReplicationEvent(*Object, Payload);
-				SceneComponent->SetWorldTransform(TransformBefore);
 				SceneComponent->UpdateComponentToWorld();
 			}
 			else
