If we are loading the p-level then skip tracking / hot reloading of the sub-level as it will be reloaded when we do the load map.

diff --git a/Engine/Plugins/Developer/Concert/ConcertSync/ConcertSyncClient/Source/ConcertSyncClient/Private/ConcertSyncClientUtil.cpp b/Engine/Plugins/Developer/Concert/ConcertSync/ConcertSyncClient/Source/ConcertSyncClient/Private/ConcertSyncClientUtil.cpp
--- a/Engine/Plugins/Developer/Concert/ConcertSync/ConcertSyncClient/Source/ConcertSyncClient/Private/ConcertSyncClientUtil.cpp
+++ b/Engine/Plugins/Developer/Concert/ConcertSync/ConcertSyncClient/Source/ConcertSyncClient/Private/ConcertSyncClientUtil.cpp
@@ -616,11 +616,14 @@ void HotReloadPackages(TArrayView<const FName> InPackageNames)
 			{
 				ExistingPackage->ClearPackageFlags(PKG_NewlyCreated);
 			}
-			ExistingPackages.Add(ExistingPackage);
 			if (ExistingPackage->ContainsMap())
 			{
 				bAddPersistentLevel = ShouldReloadPersistentLevel(ExistingPackage);
 			}
+			else
+			{
+				ExistingPackages.Add(ExistingPackage);
+			}
 		}
 	}
 
