Concert: Fix for crash when transacting level instance actor filters.

Before, when InObject was a ULevelInstanceComponent, we'd return USceneComponent::StaticClass. That happened to work for any properties which ULevelInstanceComponent inherited from USceneComponent (such as transforms, which we tested), but would crash for any properties introduced in/unique to ULevelInstanceComponent (such as actor filters, which we must not have).

diff --git a/Engine/Plugins/Developer/Concert/ConcertSync/ConcertSyncClient/Source/ConcertSyncClient/Private/ConcertClientTransactionBridge.cpp b/Engine/Plugins/Developer/Concert/ConcertSync/ConcertSyncClient/Source/ConcertSyncClient/Private/ConcertClientTransactionBridge.cpp
--- a/Engine/Plugins/Developer/Concert/ConcertSync/ConcertSyncClient/Source/ConcertSyncClient/Private/ConcertClientTransactionBridge.cpp
+++ b/Engine/Plugins/Developer/Concert/ConcertSync/ConcertSyncClient/Source/ConcertSyncClient/Private/ConcertClientTransactionBridge.cpp
@@ -859,7 +859,11 @@ UClass* GetModifiedClass(UObject* InObject)
 	{
 		return AActor::StaticClass();
 	}
-	if (Cast<ULevelInstanceComponent>(InObject) || Cast<USceneComponent>(InObject))
+	if (Cast<ULevelInstanceComponent>(InObject))
+	{
+		return ULevelInstanceComponent::StaticClass();
+	}
+	if (Cast<USceneComponent>(InObject))
 	{
 		return USceneComponent::StaticClass();
 	}
