Add a sequencer event pacer to the sequencer manager so that we don't send too many messages over multi-user.  This is controlled by two cvars:
Concert.SequencerStatePacingDuration = 0.20 (default).  This is the frequency (in seconds) that we pace the packets.
Concert.SequencerStatePacingEnabled = 1 (default), This enables the pacing packets.

Add a cvar to skip playback messages with the undefined label.  When in this state, only playing messages will get shared.  The idea is to reduce network chatter when an operator is playing.
Concert.EnableUndefinedSequencerMessages = 1 (default)

Finally to further limit network congestion, this change introduces a CVar that can force messages to send as unreliable thus could get discarded.  This cvar will be ignored ignore when sending initial play messages or rolling over the sequence in looping mode.  Note that this only
controls the operator node and the server will still replicate a state message reliably to the nDisplay cluster nodes.
Concert.SendSequencerStateEventsAsUnreliable = 0 (default)

Tested all configurations of playback with sequencer manager.  editor <-> editor, editor -> ndisplay cluster. No regressions found.

diff --git a/Engine/Plugins/Developer/Concert/ConcertSync/ConcertSyncClient/Source/ConcertSyncClient/Private/ConcertClientSequencerManager.cpp b/Engine/Plugins/Developer/Concert/ConcertSync/ConcertSyncClient/Source/ConcertSyncClient/Private/ConcertClientSequencerManager.cpp
--- a/Engine/Plugins/Developer/Concert/ConcertSync/ConcertSyncClient/Source/ConcertSyncClient/Private/ConcertClientSequencerManager.cpp
+++ b/Engine/Plugins/Developer/Concert/ConcertSync/ConcertSyncClient/Source/ConcertSyncClient/Private/ConcertClientSequencerManager.cpp
@@ -2,6 +2,7 @@
 
 #include "ConcertClientSequencerManager.h"
 #include "ConcertSequencerMessages.h"
+#include "ConcertTransportMessages.h"
 #include "Delegates/IDelegateInstance.h"
 #include "HAL/IConsoleManager.h"
 #include "Misc/AssertionMacros.h"
@@ -57,10 +58,77 @@ static TAutoConsoleVariable<int32> CVarEnableRemoteSequencerClose(TEXT("Concert.
 // Enable synchronizing the timeline of unrelated Sequencers on remote machine whenever a Sequencer state event is received, if both instances have this option on.
 static TAutoConsoleVariable<int32> CVarEnableUnrelatedTimelineSync(TEXT("Concert.EnableUnrelatedTimelineSync"), 0, TEXT("Enable syncing unrelated sequencer timeline."));
 
+// Enable the sending of undefined message types to the connected endpoints.
+// This message is only used to indicate timeline state when not in playback
+// mode. This is an experimental cvar.
+static TAutoConsoleVariable<int32> CVarEnableUndefinedSequencerMessages(TEXT("Concert.EnableUndefinedSequencerMessages"), 1, TEXT("Enable undefined sequencer messages to be sent."));
+
+// Enable the sending of messages as unreliable.  Messages of this type may get skipped if the network is congested.
+// Sequencer state events can be skipped as the sequencer will send regular ticks that would issue new state events.
+static TAutoConsoleVariable<int32> CVarSendStateEventsAsUnreliable(TEXT("Concert.SendSequencerStateEventsAsUnreliable"), 0, TEXT("Enable sending sequencer events in unreliable mode. Note this "));
+
 // Enable always closing player on remote machine whenever a sequencer is closed on an editor.
 static TAutoConsoleVariable<int32> CVarAlwaysCloseGamePlayerOnCloseEvent(
 	TEXT("Concert.AlwaysCloseGamePlayerOnCloseEvent"), 1, TEXT("Force this player to close even if other editors have it open. This CVar only works on `-game` instances."));
 
+static TAutoConsoleVariable<float> CVarSequencerStatePacingDuration(
+	TEXT("Concert.SequencerStatePacingDuration"), 0.20f, TEXT("Duration we use (in seconds) to pace sequencer state events to clients."));
+
+static TAutoConsoleVariable<bool> CVarSequencerStatePacingEnabled(
+	TEXT("Concert.SequencerStatePacingEnabled"), 1, TEXT("Use sequencer state pacing to prevent too many messages sent to the server."));
+
+
+class FConcertClientSequencerStateEventPacer
+{
+public:
+	bool ShouldIncludeMessageForSend(FConcertClientSequencerManager::EPlaybackMode Mode)
+	{
+		if (CVarEnableUndefinedSequencerMessages.GetValueOnAnyThread() == 0 && Mode == FConcertClientSequencerManager::EPlaybackMode::Undefined)
+		{
+			return false;
+		}
+		return true;
+	}
+
+	void AddStateEvent(TSharedPtr<IConcertClientSession>& Session, FConcertSequencerStateEvent InEvent, FConcertClientSequencerManager::EPlaybackMode Mode)
+	{
+		if (!ShouldIncludeMessageForSend(Mode))
+		{
+			return;
+		}
+
+		bool bShouldForceSend = bIsFirstEvent || CVarSequencerStatePacingEnabled.GetValueOnAnyThread() == 0;
+		if (!bShouldForceSend && InEvent.State.bLoopMode && CurrentStateEvent)
+		{
+			// If we have loop enabled and we have started back over again then we should force a send.
+			bShouldForceSend = InEvent.State.Time.AsSeconds() < CurrentStateEvent->State.Time.AsSeconds();
+		}
+		CurrentStateEvent = MoveTemp(InEvent);
+		Tick(Session, bShouldForceSend);
+	}
+
+	void Tick(TSharedPtr<IConcertClientSession>& Session, bool bForceSend = false)
+	{
+		if (CurrentStateEvent)
+		{
+			double CurrentTime = FPlatformTime::Seconds();
+			const double DeltaTime = CurrentTime - LastSendTime;
+			if (bForceSend || DeltaTime > CVarSequencerStatePacingDuration.GetValueOnGameThread())
+			{
+				bool bShouldSendReliably = bForceSend && CVarSequencerStatePacingEnabled.GetValueOnAnyThread();
+				EConcertMessageFlags Flags = bShouldSendReliably || CVarSendStateEventsAsUnreliable.GetValueOnAnyThread() == 0 ? EConcertMessageFlags::ReliableOrdered : EConcertMessageFlags::None;
+				Session->SendCustomEvent(*CurrentStateEvent, Session->GetSessionServerEndpointId(), Flags);
+				LastSendTime = CurrentTime;
+				CurrentStateEvent.Reset();
+			}
+		}
+	}
+
+private:
+	TOptional<FConcertSequencerStateEvent> CurrentStateEvent;
+	double LastSendTime = 0;
+	bool bIsFirstEvent = true;
+};
 
 class FConcertClientSequencePreloader : public TSharedFromThis<FConcertClientSequencePreloader>
 {
@@ -113,8 +181,11 @@ FConcertClientSequencerManager::FConcertClientSequencerManager(IConcertSyncClien
 
 	bRespondingToTransportEvent = false;
 
-	ISequencerModule& SequencerModule = FModuleManager::Get().LoadModuleChecked<ISequencerModule>("Sequencer");
-	OnSequencerCreatedHandle = SequencerModule.RegisterOnSequencerCreated(FOnSequencerCreated::FDelegate::CreateRaw(this, &FConcertClientSequencerManager::OnSequencerCreated));
+	if (GIsEditor)
+	{
+		ISequencerModule& SequencerModule = FModuleManager::Get().LoadModuleChecked<ISequencerModule>("Sequencer");
+		OnSequencerCreatedHandle = SequencerModule.RegisterOnSequencerCreated(FOnSequencerCreated::FDelegate::CreateRaw(this, &FConcertClientSequencerManager::OnSequencerCreated));
+	}
 
 	FCoreUObjectDelegates::OnPackageReloaded.AddRaw(this, &FConcertClientSequencerManager::HandleAssetReload);
 }
@@ -131,10 +202,13 @@ bool IsPendingTakePath(const FString& InSequencePath)
 
 FConcertClientSequencerManager::~FConcertClientSequencerManager()
 {
-	ISequencerModule* SequencerModulePtr = FModuleManager::Get().GetModulePtr<ISequencerModule>("Sequencer");
-	if (SequencerModulePtr)
+	if (GIsEditor)
 	{
-		SequencerModulePtr->UnregisterOnSequencerCreated(OnSequencerCreatedHandle);
+		ISequencerModule* SequencerModulePtr = FModuleManager::Get().GetModulePtr<ISequencerModule>("Sequencer");
+		if (SequencerModulePtr)
+		{
+			SequencerModulePtr->UnregisterOnSequencerCreated(OnSequencerCreatedHandle);
+		}
 	}
 
 	SetActiveWorkspace(nullptr);
@@ -176,7 +250,8 @@ void FConcertClientSequencerManager::OnSequencerCreated(TSharedRef<ISequencer> I
 	OpenSequencer.PlaybackMode = EPlaybackMode::Undefined;
 	OpenSequencer.OnGlobalTimeChangedHandle = InSequencer->OnGlobalTimeChanged().AddRaw(this, &FConcertClientSequencerManager::OnSequencerTimeChanged, OpenSequencer.WeakSequencer);
 	OpenSequencer.OnCloseEventHandle = InSequencer->OnCloseEvent().AddRaw(this, &FConcertClientSequencerManager::OnSequencerClosed);
-	int32 OpenIndex = OpenSequencers.Add(OpenSequencer);
+	OpenSequencer.StateEventPacer = MakePimpl<FConcertClientSequencerStateEventPacer>();
+	int32 OpenIndex = OpenSequencers.Add(MoveTemp(OpenSequencer));
 
 	// Setup stored state.  Since this is an open event by the sequencer we might be the controller and should not set player
 	// state from previous stored state.  Hence playback mode is set to undefined.  We should wait for a global time event from
@@ -495,6 +570,7 @@ void FConcertClientSequencerManager::OnSequencerClosed(TSharedRef<ISequencer> In
 				FConcertSequencerCloseEvent CloseEvent;
 				CloseEvent.bControllerClose = ClosingSequencer.PlaybackMode == EPlaybackMode::Controller; // this sequencer had control over the sequence playback
 				CloseEvent.SequenceObjectPath = SequenceObjectPath;
+				ClosingSequencer.StateEventPacer->Tick(Session, true /*bForceSend*/);
 				Session->SendCustomEvent(CloseEvent, Session->GetSessionServerEndpointId(), EConcertMessageFlags::ReliableOrdered);
 			}
 		}
@@ -598,7 +674,8 @@ void FConcertClientSequencerManager::OnSequencerTimeChanged(TWeakPtr<ISequencer>
 			// Send to client and server
 			UE_LOG(LogConcertSequencerSync, Verbose, TEXT("Sending StateEvent: %s, at frame: %d"),
 				*SequencerStateEvent.State.SequenceObjectPath, SequencerStateEvent.State.Time.Time.FrameNumber.Value);
-			Session->SendCustomEvent(SequencerStateEvent, Session->GetSessionServerEndpointId(), EConcertMessageFlags::ReliableOrdered);
+
+			OpenSequencer->StateEventPacer->AddStateEvent(Session, SequencerStateEvent, OpenSequencer->PlaybackMode);
 
 			// If we're playing then ensure we are set to controller (driving the playback on all clients)
 			if (SequencerStateEvent.State.PlayerStatus == EConcertMovieScenePlayerStatus::Playing)
@@ -1353,6 +1430,20 @@ void FConcertClientSequencerManager::OnWorkspaceEndFrameCompleted()
 		ApplyTimeAdjustmentEvent(Event);
 	}
 	PendingTimeAdjustmentEvents.Reset();
+
+	if (GEditor)
+	{
+		if (TSharedPtr<IConcertClientSession> Session = WeakSession.Pin())
+		{
+			for (FOpenSequencerData& Entry : OpenSequencers)
+			{
+				if (Entry.StateEventPacer)
+				{
+					Entry.StateEventPacer->Tick(Session);
+				}
+			}
+		}
+	}
 }
 
 void FConcertClientSequencerManager::AddReferencedObjects(FReferenceCollector& Collector)
diff --git a/Engine/Plugins/Developer/Concert/ConcertSync/ConcertSyncClient/Source/ConcertSyncClient/Private/ConcertClientSequencerManager.h b/Engine/Plugins/Developer/Concert/ConcertSync/ConcertSyncClient/Source/ConcertSyncClient/Private/ConcertClientSequencerManager.h
--- a/Engine/Plugins/Developer/Concert/ConcertSync/ConcertSyncClient/Source/ConcertSyncClient/Private/ConcertClientSequencerManager.h
+++ b/Engine/Plugins/Developer/Concert/ConcertSync/ConcertSyncClient/Source/ConcertSyncClient/Private/ConcertClientSequencerManager.h
@@ -8,6 +8,7 @@
 #include "IConcertClientSequencerManager.h"
 #include "ConcertSequencerMessages.h"
 #include "UObject/WeakObjectPtrTemplates.h"
+#include "Templates/PimplPtr.h"
 
 struct FConcertSessionContext;
 class IConcertSyncClient;
@@ -21,6 +22,7 @@ class FConcertClientSequencePreloader;
 class FConcertClientWorkspace;
 struct FSequencerInitParams;
 class ULevelSequence;
+class FConcertClientSequencerStateEventPacer;
 
 /**
  * Sequencer manager that is held by the client sync module that keeps track of open sequencer UIs, regardless of whether a session is open or not
@@ -140,6 +142,9 @@ private:
 
 		/** Delegate handle to the Close event for the sequencer, if locally opened. */
 		FDelegateHandle OnCloseEventHandle;
+
+		/** Pointer event pacer for state events. This class will pace state events so that we do not overwhelm the connection with too many. */
+		TPimplPtr<FConcertClientSequencerStateEventPacer> StateEventPacer;
 	};
 
 private:
@@ -370,6 +375,8 @@ private:
 
 	/** Weak pointer to the active workspace. */
 	TWeakPtr<FConcertClientWorkspace> Workspace;
+
+	friend class FConcertClientSequencerStateEventPacer;
 };
 
 #endif // WITH_EDITOR
