Set GIsTransacting when calling Pre/PostEditChange/Undo when applying a Concert transaction

#rnx

diff --git a/Engine/Plugins/Developer/Concert/ConcertSync/ConcertSyncClient/Source/ConcertSyncClient/Private/ConcertClientTransactionBridge.cpp b/Engine/Plugins/Developer/Concert/ConcertSync/ConcertSyncClient/Source/ConcertSyncClient/Private/ConcertClientTransactionBridge.cpp
--- a/Engine/Plugins/Developer/Concert/ConcertSync/ConcertSyncClient/Source/ConcertSyncClient/Private/ConcertClientTransactionBridge.cpp
+++ b/Engine/Plugins/Developer/Concert/ConcertSync/ConcertSyncClient/Source/ConcertSyncClient/Private/ConcertClientTransactionBridge.cpp
@@ -493,7 +493,10 @@ void ProcessTransactionEvent(const FConcertTransactionEventBase& InEvent, const
 			const bool bLevelIsDirty = Level ? Level->GetPackage()->IsDirty() : false;
 
 			// Transaction annotations require us to invoke the redo flow (even for snapshots!) as that's the only thing that can apply the annotation
-			TransactionObject->PreEditUndo();
+			{
+				TGuardValue<bool> IsTransactingGuard(GIsTransacting, true);
+				TransactionObject->PreEditUndo();
+			}
 
 			// Levels are immune from dirty changes when using external objects.  See ULevel::PreEditUndo() If we
 			// modified any dirty flags as a result of the PreEditUndo then restore it back here as if it didn't happen.
@@ -513,6 +516,7 @@ void ProcessTransactionEvent(const FConcertTransactionEventBase& InEvent, const
 			{
 				if (bIsSnapshot)
 				{
+					TGuardValue<bool> IsTransactingGuard(GIsTransacting, true);
 					// Prevent FlushRenderingCommands from running when in -game. It can cause performance issues / hitching during user interaction.
 					TGuardValue<bool> ShouldFlushRenderingCommands(GFlushRenderingCommandsOnPreEditChange, GIsEditor);					
 					TransactionObject->PreEditChange(TransactionProp);
@@ -598,6 +602,7 @@ void ProcessTransactionEvent(const FConcertTransactionEventBase& InEvent, const
 			{
 				if (bIsSnapshot)
 				{
+					TGuardValue<bool> IsTransactingGuard(GIsTransacting, true);
 					TransactionObject->PostEditChange();
 				}
 
@@ -611,10 +616,12 @@ void ProcessTransactionEvent(const FConcertTransactionEventBase& InEvent, const
 		if (TransactionAnnotation)
 		{
 			// Transaction annotations require us to invoke the redo flow (even for snapshots!) as that's the only thing that can apply the annotation
+			TGuardValue<bool> IsTransactingGuard(GIsTransacting, true);
 			TransactionObject->PostEditUndo(TransactionAnnotation);
 		}
 		else if (!bIsSnapshot)
 		{
+			TGuardValue<bool> IsTransactingGuard(GIsTransacting, true);
 			TransactionObject->PostEditUndo();
 		}
 
