Refactor TSharedRef/Ptr inputs to use references tagged with UE_LITETIMEBOUND.
This is supposed to communicate intend and simplify destruction logic.

diff --git a/Engine/Plugins/Developer/Concert/ConcertSync/ConcertSyncClient/Source/ConcertSyncClient/Private/ConcertSyncClient.cpp b/Engine/Plugins/Developer/Concert/ConcertSync/ConcertSyncClient/Source/ConcertSyncClient/Private/ConcertSyncClient.cpp
--- a/Engine/Plugins/Developer/Concert/ConcertSync/ConcertSyncClient/Source/ConcertSyncClient/Private/ConcertSyncClient.cpp
+++ b/Engine/Plugins/Developer/Concert/ConcertSync/ConcertSyncClient/Source/ConcertSyncClient/Private/ConcertSyncClient.cpp
@@ -187,9 +187,9 @@ void FConcertSyncClient::RegisterConcertSyncHandlers(TSharedRef<IConcertClientSe
 		CreateWorkspace(LiveSession.ToSharedRef());
 		
 		// Create Replication Manager
-		if (EnumHasAnyFlags(LiveSession->GetSessionFlags(), EConcertSyncSessionFlags::EnableReplication))
+		if (EnumHasAnyFlags(LiveSession->GetSessionFlags(), EConcertSyncSessionFlags::EnableReplication) && ensure(Bridges.ReplicationBridge))
 		{
-			ReplicationManager = MakeUnique<UE::ConcertSyncClient::Replication::FReplicationManager>(InSession, Bridges.ReplicationBridge);
+			ReplicationManager = MakeUnique<UE::ConcertSyncClient::Replication::FReplicationManager>(InSession, *Bridges.ReplicationBridge);
 			ReplicationManager->StartAcceptingJoinRequests();
 		}
 
diff --git a/Engine/Plugins/Developer/Concert/ConcertSync/ConcertSyncClient/Source/ConcertSyncClient/Private/Replication/Manager/ReplicationManager.cpp b/Engine/Plugins/Developer/Concert/ConcertSync/ConcertSyncClient/Source/ConcertSyncClient/Private/Replication/Manager/ReplicationManager.cpp
--- a/Engine/Plugins/Developer/Concert/ConcertSync/ConcertSyncClient/Source/ConcertSyncClient/Private/Replication/Manager/ReplicationManager.cpp
+++ b/Engine/Plugins/Developer/Concert/ConcertSync/ConcertSyncClient/Source/ConcertSyncClient/Private/Replication/Manager/ReplicationManager.cpp
@@ -10,7 +10,7 @@
 
 namespace UE::ConcertSyncClient::Replication
 {
-	FReplicationManager::FReplicationManager(TSharedRef<IConcertClientSession> InLiveSession, IConcertClientReplicationBridge* InBridge)
+	FReplicationManager::FReplicationManager(TSharedRef<IConcertClientSession> InLiveSession, IConcertClientReplicationBridge& InBridge)
 		: Session(MoveTemp(InLiveSession))
 		, Bridge(InBridge)
 	{}
diff --git a/Engine/Plugins/Developer/Concert/ConcertSync/ConcertSyncClient/Source/ConcertSyncClient/Private/Replication/Manager/ReplicationManager.h b/Engine/Plugins/Developer/Concert/ConcertSync/ConcertSyncClient/Source/ConcertSyncClient/Private/Replication/Manager/ReplicationManager.h
--- a/Engine/Plugins/Developer/Concert/ConcertSync/ConcertSyncClient/Source/ConcertSyncClient/Private/Replication/Manager/ReplicationManager.h
+++ b/Engine/Plugins/Developer/Concert/ConcertSync/ConcertSyncClient/Source/ConcertSyncClient/Private/Replication/Manager/ReplicationManager.h
@@ -18,7 +18,7 @@ namespace UE::ConcertSyncClient::Replication
 		friend FReplicationManagerState; 
 	public:
 		
-		FReplicationManager(TSharedRef<IConcertClientSession> InLiveSession, IConcertClientReplicationBridge* InBridge);
+		FReplicationManager(TSharedRef<IConcertClientSession> InLiveSession, IConcertClientReplicationBridge& InBridge UE_LIFETIMEBOUND);
 		virtual ~FReplicationManager() override;
 
 		/** Starts accepting join requests. Must be called separately from constructor because of TSharedFromThis asserting if SharedThis is called in constructor. */
@@ -46,7 +46,7 @@ namespace UE::ConcertSyncClient::Replication
 		/** Session instance this manager was created for. */
 		TSharedRef<IConcertClientSession> Session;
 		/** The replication bridge is responsible for applying received data and generating data to send. */
-		IConcertClientReplicationBridge* Bridge;
+		IConcertClientReplicationBridge& Bridge;
 
 		/** The current state this manager is in, e.g. waiting for connection request, connecting, connected, etc. */
 		TSharedPtr<FReplicationManagerState> CurrentState;
diff --git a/Engine/Plugins/Developer/Concert/ConcertSync/ConcertSyncClient/Source/ConcertSyncClient/Private/Replication/Manager/ReplicationManagerState_Connected.cpp b/Engine/Plugins/Developer/Concert/ConcertSync/ConcertSyncClient/Source/ConcertSyncClient/Private/Replication/Manager/ReplicationManagerState_Connected.cpp
--- a/Engine/Plugins/Developer/Concert/ConcertSync/ConcertSyncClient/Source/ConcertSyncClient/Private/Replication/Manager/ReplicationManagerState_Connected.cpp
+++ b/Engine/Plugins/Developer/Concert/ConcertSync/ConcertSyncClient/Source/ConcertSyncClient/Private/Replication/Manager/ReplicationManagerState_Connected.cpp
@@ -70,7 +70,7 @@ namespace UE::ConcertSyncClient::Replication
 	
 	FReplicationManagerState_Connected::FReplicationManagerState_Connected(
 		TSharedRef<IConcertClientSession> LiveSession,
-		IConcertClientReplicationBridge* ReplicationBridge,
+		IConcertClientReplicationBridge& ReplicationBridge,
 		TArray<FConcertReplicationStream> StreamDescriptions,
 		FReplicationManager& Owner
 		)
@@ -79,24 +79,24 @@ namespace UE::ConcertSyncClient::Replication
 		, ReplicationBridge(ReplicationBridge)
 		, RegisteredStreams(MoveTemp(StreamDescriptions))
 		// TODO DP: Use config to determine which replication format to use
-		, ReplicationFormat(MakeShared<ConcertSyncCore::FFullObjectFormat>())
-		, ReplicationDataSource(MakeShared<FClientReplicationDataCollector>(
+		, ReplicationFormat(MakeUnique<ConcertSyncCore::FFullObjectFormat>())
+		, ReplicationDataSource(
 			ReplicationBridge,
-			ReplicationFormat,
+			*ReplicationFormat,
 			FClientReplicationDataCollector::FGetClientStreams::CreateLambda([this]()
 			{
 				return &RegisteredStreams;
 			}),
 			LiveSession->GetSessionClientEndpointId()
-			))
+			)
 		, Sender(
 			ConcertSyncCore::FGetObjectFrequencySettings::CreateRaw(this, &FReplicationManagerState_Connected::GetObjectFrequencySettings),
-			LiveSession->GetSessionServerEndpointId(), LiveSession, ReplicationDataSource
+			LiveSession->GetSessionServerEndpointId(), *LiveSession, ReplicationDataSource
 			)
-		, ReceivedDataCache(MakeShared<ConcertSyncCore::FObjectReplicationCache>(ReplicationFormat))
-		, Receiver(MakeShared<ConcertSyncCore::FObjectReplicationReceiver>(LiveSession, ReceivedDataCache))
-		, ReceivedReplicationQueuer(FClientReplicationDataQueuer::Make(ReplicationBridge, ReceivedDataCache))
-		, ReplicationApplier(MakeShared<FObjectReplicationApplierProcessor>(ReplicationBridge, ReplicationFormat, ReceivedReplicationQueuer))
+		, ReceivedDataCache(MakeShared<ConcertSyncCore::FObjectReplicationCache>(*ReplicationFormat))
+		, Receiver(*LiveSession, *ReceivedDataCache)
+		, ReceivedReplicationQueuer(FClientReplicationDataQueuer::Make(ReplicationBridge, *ReceivedDataCache))
+		, ReplicationApplier(ReplicationBridge, *ReplicationFormat, *ReceivedReplicationQueuer)
 	{}
 
 	FReplicationManagerState_Connected::~FReplicationManagerState_Connected()
@@ -225,13 +225,13 @@ namespace UE::ConcertSyncClient::Replication
 	{
 		TSet<FGuid> Result;
 		int32 ExpectedNumStreams = 0;
-		ReplicationDataSource->ForEachOwnedObject([this, &Callback, &Result, &ExpectedNumStreams](const FSoftObjectPath& ObjectPath)
+		ReplicationDataSource.ForEachOwnedObject([this, &Callback, &Result, &ExpectedNumStreams](const FSoftObjectPath& ObjectPath)
 		{
 			// Reuse TSet (if possible) for a slightly better memory footprint
 			ExpectedNumStreams = FMath::Max(ExpectedNumStreams, Result.Num());
 			Result.Empty(Result.Num());
 			
-			ReplicationDataSource->AppendOwningStreamsForObject(ObjectPath, Result);
+			ReplicationDataSource.AppendOwningStreamsForObject(ObjectPath, Result);
 			return Callback(ObjectPath, MoveTemp(Result));
 		});
 		return EAuthorityEnumerationResult::Iterated;
@@ -242,7 +242,7 @@ namespace UE::ConcertSyncClient::Replication
 		) const
 	{
 		TSet<FGuid> Result;
-		ReplicationDataSource->AppendOwningStreamsForObject(ObjectPath, Result);
+		ReplicationDataSource.AppendOwningStreamsForObject(ObjectPath, Result);
 		return Result;
 	}
 
@@ -256,7 +256,7 @@ namespace UE::ConcertSyncClient::Replication
 		// TODO UE-190714: We should set a time budget for the client so ticking does not cause frame spikes
 		const ConcertSyncCore::FProcessObjectsParams Params { DeltaTime };
 		Sender.ProcessObjects(Params);
-		ReplicationApplier->ProcessObjects(Params);
+		ReplicationApplier.ProcessObjects(Params);
 	}
 
 	void FReplicationManagerState_Connected::UpdateReplicatedObjectsAfterStreamChange(const FConcertReplication_ChangeStream_Request& Request, const FConcertReplication_ChangeStream_Response& Response)
@@ -290,7 +290,7 @@ namespace UE::ConcertSyncClient::Replication
 		
 		for (const TPair<FSoftObjectPath, TArray<FGuid>>& ModifiedObject : BundledModifiedObjects)
 		{
-			ReplicationDataSource->OnObjectStreamModified(ModifiedObject.Key, ModifiedObject.Value);
+			ReplicationDataSource.OnObjectStreamModified(ModifiedObject.Key, ModifiedObject.Value);
 		}
 	}
 
@@ -310,23 +310,23 @@ namespace UE::ConcertSyncClient::Replication
 		}
 	}
 
-	void FReplicationManagerState_Connected::HandleRemovingReplicatedObjects(const FConcertReplication_ChangeStream_Request& Request) const
+	void FReplicationManagerState_Connected::HandleRemovingReplicatedObjects(const FConcertReplication_ChangeStream_Request& Request)
 	{
 		Private::ForEachObjectRemovedFromStreams(Request, [this](const FSoftObjectPath& Object, const TArray<FGuid>& RemovedStreams)
 		{
-			ReplicationDataSource->RemoveReplicatedObjectStreams(Object, RemovedStreams);
+			ReplicationDataSource.RemoveReplicatedObjectStreams(Object, RemovedStreams);
 		});
 	}
 
-	void FReplicationManagerState_Connected::RevertRemovingReplicatedObjects(const FConcertReplication_ChangeStream_Request& Request) const
+	void FReplicationManagerState_Connected::RevertRemovingReplicatedObjects(const FConcertReplication_ChangeStream_Request& Request)
 	{
 		Private::ForEachObjectRemovedFromStreams(Request, [this](const FSoftObjectPath& Object, const TArray<FGuid>& RemovedStreams)
 		{
-			ReplicationDataSource->AddReplicatedObjectStreams(Object, RemovedStreams);
+			ReplicationDataSource.AddReplicatedObjectStreams(Object, RemovedStreams);
 		});
 	}
 
-	void FReplicationManagerState_Connected::UpdateReplicatedObjectsAfterAuthorityChange(FConcertReplication_ChangeAuthority_Request&& Request, const FConcertReplication_ChangeAuthority_Response& Response) const
+	void FReplicationManagerState_Connected::UpdateReplicatedObjectsAfterAuthorityChange(FConcertReplication_ChangeAuthority_Request&& Request, const FConcertReplication_ChangeAuthority_Response& Response)
 	{
 		OnPreAuthorityChangedDelegate.Broadcast(Request, { Response });
 		ON_SCOPE_EXIT{ OnPostAuthorityChangedDelegate.Broadcast(); };
@@ -351,24 +351,24 @@ namespace UE::ConcertSyncClient::Replication
 			const bool bWasFullyRejected = ReplicatedStreams.StreamIds.IsEmpty(); 
 			if (!bWasFullyRejected)
 			{
-				ReplicationDataSource->AddReplicatedObjectStreams(TakeAuthority.Key, ReplicatedStreams.StreamIds);
+				ReplicationDataSource.AddReplicatedObjectStreams(TakeAuthority.Key, ReplicatedStreams.StreamIds);
 			}
 		}
 	}
 
-	void FReplicationManagerState_Connected::HandleReleasingReplicatedObjects(const FConcertReplication_ChangeAuthority_Request& Request) const
+	void FReplicationManagerState_Connected::HandleReleasingReplicatedObjects(const FConcertReplication_ChangeAuthority_Request& Request)
 	{
 		for (const TPair<FSoftObjectPath, FConcertStreamArray>& ReleaseAuthority : Request.ReleaseAuthority)
 		{
-			ReplicationDataSource->RemoveReplicatedObjectStreams(ReleaseAuthority.Key, ReleaseAuthority.Value.StreamIds);
+			ReplicationDataSource.RemoveReplicatedObjectStreams(ReleaseAuthority.Key, ReleaseAuthority.Value.StreamIds);
 		}
 	}
 
-	void FReplicationManagerState_Connected::RevertReleasingReplicatedObjects(const FConcertReplication_ChangeAuthority_Request& Request) const
+	void FReplicationManagerState_Connected::RevertReleasingReplicatedObjects(const FConcertReplication_ChangeAuthority_Request& Request)
 	{
 		for (const TPair<FSoftObjectPath, FConcertStreamArray>& ReleaseAuthority : Request.ReleaseAuthority)
 		{
-			ReplicationDataSource->AddReplicatedObjectStreams(ReleaseAuthority.Key, ReleaseAuthority.Value.StreamIds);
+			ReplicationDataSource.AddReplicatedObjectStreams(ReleaseAuthority.Key, ReleaseAuthority.Value.StreamIds);
 		}
 	}
 
diff --git a/Engine/Plugins/Developer/Concert/ConcertSync/ConcertSyncClient/Source/ConcertSyncClient/Private/Replication/Manager/ReplicationManagerState_Connected.h b/Engine/Plugins/Developer/Concert/ConcertSync/ConcertSyncClient/Source/ConcertSyncClient/Private/Replication/Manager/ReplicationManagerState_Connected.h
--- a/Engine/Plugins/Developer/Concert/ConcertSync/ConcertSyncClient/Source/ConcertSyncClient/Private/Replication/Manager/ReplicationManagerState_Connected.h
+++ b/Engine/Plugins/Developer/Concert/ConcertSync/ConcertSyncClient/Source/ConcertSyncClient/Private/Replication/Manager/ReplicationManagerState_Connected.h
@@ -3,23 +3,20 @@
 #pragma once
 
 #include "ReplicationManagerState.h"
+#include "Replication/Formats/IObjectReplicationFormat.h"
 #include "Replication/Processing/ClientReplicationDataCollector.h"
 #include "Replication/Processing/ClientReplicationDataQueuer.h"
+#include "Replication/Processing/ObjectReplicationApplierProcessor.h"
+#include "Replication/Processing/ObjectReplicationReceiver.h"
 #include "Replication/Processing/ObjectReplicationSender.h"
 #include "Replication/Processing/Proxy/ObjectProcessorProxy_Frequency.h"
 
 class IConcertClientSession;
 
-namespace UE::ConcertSyncCore
-{
-	class FObjectReplicationReceiver;
-}
+namespace UE::ConcertSyncCore { class FObjectReplicationReceiver; }
 
 namespace UE::ConcertSyncClient::Replication
 {
-	class FObjectReplicationApplierProcessor;
-	class FObjectReplicationSender;
-	
 	/**
 	 * State for when the client has successfully completed a replication handshake.
 	 *
@@ -33,9 +30,9 @@ namespace UE::ConcertSyncClient::Replication
 
 		FReplicationManagerState_Connected(
 			TSharedRef<IConcertClientSession> LiveSession,
-			IConcertClientReplicationBridge* ReplicationBridge,
+			IConcertClientReplicationBridge& ReplicationBridge UE_LIFETIMEBOUND,
 			TArray<FConcertReplicationStream> StreamDescriptions,
-			FReplicationManager& Owner
+			FReplicationManager& Owner UE_LIFETIMEBOUND
 			);
 		virtual ~FReplicationManagerState_Connected() override;
 
@@ -57,16 +54,16 @@ namespace UE::ConcertSyncClient::Replication
 		/** Passed to FReplicationManagerState_Disconnected */
 		const TSharedRef<IConcertClientSession> LiveSession;
 		/** Passed to FReplicationManagerState_Disconnected */
-		IConcertClientReplicationBridge* const ReplicationBridge;
+		IConcertClientReplicationBridge& ReplicationBridge;
 		/** The streams this client has registered with the server. */
 		TArray<FConcertReplicationStream> RegisteredStreams;
 		
 		/** The format this client will use for sending & receiving data. */
-		const TSharedRef<ConcertSyncCore::IObjectReplicationFormat> ReplicationFormat;
+		const TUniquePtr<ConcertSyncCore::IObjectReplicationFormat> ReplicationFormat;
 
 		// Sending
 		/** Used as source of replication data. */
-		const TSharedRef<FClientReplicationDataCollector> ReplicationDataSource;
+		FClientReplicationDataCollector ReplicationDataSource;
 		
 		/** Sends to remote endpoint and makes sure the objects are replicated at the specified frequency settings. */
 		using FDataRelayThrottledByFrequency = ConcertSyncCore::TObjectProcessorProxy_Frequency<ConcertSyncCore::FObjectReplicationSender>;
@@ -77,11 +74,11 @@ namespace UE::ConcertSyncClient::Replication
 		/** Stores data received by Receiver until it is consumed by ReceivedReplicationQueuer. */
 		const TSharedRef<ConcertSyncCore::FObjectReplicationCache> ReceivedDataCache;
 		/** Receives data from remote endpoints via message bus.  */
-		const TSharedRef<ConcertSyncCore::FObjectReplicationReceiver> Receiver;
-		/** Queues data until is can be processed. */
+		ConcertSyncCore::FObjectReplicationReceiver Receiver;
+		/** Queues data until is can be processed. Shared because FObjectReplicationCache API expects it. */
 		const TSharedRef<FClientReplicationDataQueuer> ReceivedReplicationQueuer;
 		/** Processes data from ReceivedReplicationQueuer once we tick. */
-		const TSharedRef<FObjectReplicationApplierProcessor> ReplicationApplier;
+		FObjectReplicationApplierProcessor ReplicationApplier;
 
 		//~ Begin FReplicationManagerState Interface
 		virtual void OnEnterState() override;
@@ -101,8 +98,8 @@ namespace UE::ConcertSyncClient::Replication
 		
 		/** Updates replicated objects affected by the change request. */
 		void UpdateReplicatedObjectsAfterStreamChange(const FConcertReplication_ChangeStream_Request& Request, const FConcertReplication_ChangeStream_Response& Response);
-		void HandleRemovingReplicatedObjects(const FConcertReplication_ChangeStream_Request& Request) const;
-		void RevertRemovingReplicatedObjects(const FConcertReplication_ChangeStream_Request& Request) const;
+		void HandleRemovingReplicatedObjects(const FConcertReplication_ChangeStream_Request& Request);
+		void RevertRemovingReplicatedObjects(const FConcertReplication_ChangeStream_Request& Request);
 
 		/**
 		 * Updates the objects which should be replicated after changing authority.
@@ -110,9 +107,9 @@ namespace UE::ConcertSyncClient::Replication
 		 * @note Request is accepted as && because this function rewrites its memory when looking at rejections.
 		 * Since the request was already sent to the server it is assumed the request can just contain trash after.
 		 */
-		void UpdateReplicatedObjectsAfterAuthorityChange(FConcertReplication_ChangeAuthority_Request&& Request, const FConcertReplication_ChangeAuthority_Response& Response) const;
-		void HandleReleasingReplicatedObjects(const FConcertReplication_ChangeAuthority_Request& Request) const;
-		void RevertReleasingReplicatedObjects(const FConcertReplication_ChangeAuthority_Request& Request) const;
+		void UpdateReplicatedObjectsAfterAuthorityChange(FConcertReplication_ChangeAuthority_Request&& Request, const FConcertReplication_ChangeAuthority_Response& Response);
+		void HandleReleasingReplicatedObjects(const FConcertReplication_ChangeAuthority_Request& Request);
+		void RevertReleasingReplicatedObjects(const FConcertReplication_ChangeAuthority_Request& Request);
 		
 		/** Callback to Sender for obtaining an object's frequency settings. */
 		FConcertObjectReplicationSettings GetObjectFrequencySettings(const FConcertReplicatedObjectId& Object) const;
diff --git a/Engine/Plugins/Developer/Concert/ConcertSync/ConcertSyncClient/Source/ConcertSyncClient/Private/Replication/Manager/ReplicationManagerState_Disconnected.cpp b/Engine/Plugins/Developer/Concert/ConcertSync/ConcertSyncClient/Source/ConcertSyncClient/Private/Replication/Manager/ReplicationManagerState_Disconnected.cpp
--- a/Engine/Plugins/Developer/Concert/ConcertSync/ConcertSyncClient/Source/ConcertSyncClient/Private/Replication/Manager/ReplicationManagerState_Disconnected.cpp
+++ b/Engine/Plugins/Developer/Concert/ConcertSync/ConcertSyncClient/Source/ConcertSyncClient/Private/Replication/Manager/ReplicationManagerState_Disconnected.cpp
@@ -17,8 +17,8 @@ namespace UE::ConcertSyncClient::Replication
 	
 	FReplicationManagerState_Disconnected::FReplicationManagerState_Disconnected(
 		TSharedRef<IConcertClientSession> LiveSession,
-		IConcertClientReplicationBridge* ReplicationBridge,
-			FReplicationManager& Owner
+		IConcertClientReplicationBridge& ReplicationBridge,
+		FReplicationManager& Owner
 		)
 		: FReplicationManagerState(Owner)
 		, LiveSession(MoveTemp(LiveSession))
diff --git a/Engine/Plugins/Developer/Concert/ConcertSync/ConcertSyncClient/Source/ConcertSyncClient/Private/Replication/Manager/ReplicationManagerState_Disconnected.h b/Engine/Plugins/Developer/Concert/ConcertSync/ConcertSyncClient/Source/ConcertSyncClient/Private/Replication/Manager/ReplicationManagerState_Disconnected.h
--- a/Engine/Plugins/Developer/Concert/ConcertSync/ConcertSyncClient/Source/ConcertSyncClient/Private/Replication/Manager/ReplicationManagerState_Disconnected.h
+++ b/Engine/Plugins/Developer/Concert/ConcertSync/ConcertSyncClient/Source/ConcertSyncClient/Private/Replication/Manager/ReplicationManagerState_Disconnected.h
@@ -17,8 +17,8 @@ namespace UE::ConcertSyncClient::Replication
 
 		FReplicationManagerState_Disconnected(
 			TSharedRef<IConcertClientSession> LiveSession,
-			IConcertClientReplicationBridge* ReplicationBridge,
-			FReplicationManager& Owner
+			IConcertClientReplicationBridge& ReplicationBridge UE_LIFETIMEBOUND,
+			FReplicationManager& Owner UE_LIFETIMEBOUND
 			);
 
 		//~ Begin IConcertClientReplicationManager Interface
@@ -33,6 +33,6 @@ namespace UE::ConcertSyncClient::Replication
 		/** Passed to FReplicationManagerState_Handshaking */
 		TSharedRef<IConcertClientSession> LiveSession;
 		/** Passed to FReplicationManagerState_Handshaking */
-		IConcertClientReplicationBridge* ReplicationBridge;
+		IConcertClientReplicationBridge& ReplicationBridge;
 	};
 }
\ No newline at end of file
diff --git a/Engine/Plugins/Developer/Concert/ConcertSync/ConcertSyncClient/Source/ConcertSyncClient/Private/Replication/Manager/ReplicationManagerState_Handshaking.cpp b/Engine/Plugins/Developer/Concert/ConcertSync/ConcertSyncClient/Source/ConcertSyncClient/Private/Replication/Manager/ReplicationManagerState_Handshaking.cpp
--- a/Engine/Plugins/Developer/Concert/ConcertSync/ConcertSyncClient/Source/ConcertSyncClient/Private/Replication/Manager/ReplicationManagerState_Handshaking.cpp
+++ b/Engine/Plugins/Developer/Concert/ConcertSync/ConcertSyncClient/Source/ConcertSyncClient/Private/Replication/Manager/ReplicationManagerState_Handshaking.cpp
@@ -15,7 +15,7 @@ namespace UE::ConcertSyncClient::Replication
 		FJoinReplicatedSessionArgs RequestArgs,
 		TPromise<FJoinReplicatedSessionResult> JoinSessionPromise,
 		TSharedRef<IConcertClientSession> LiveSession,
-		IConcertClientReplicationBridge* ReplicationBridge,
+		IConcertClientReplicationBridge& ReplicationBridge,
 		FReplicationManager& Owner
 		)
 		: FReplicationManagerState(Owner)
diff --git a/Engine/Plugins/Developer/Concert/ConcertSync/ConcertSyncClient/Source/ConcertSyncClient/Private/Replication/Manager/ReplicationManagerState_Handshaking.h b/Engine/Plugins/Developer/Concert/ConcertSync/ConcertSyncClient/Source/ConcertSyncClient/Private/Replication/Manager/ReplicationManagerState_Handshaking.h
--- a/Engine/Plugins/Developer/Concert/ConcertSync/ConcertSyncClient/Source/ConcertSyncClient/Private/Replication/Manager/ReplicationManagerState_Handshaking.h
+++ b/Engine/Plugins/Developer/Concert/ConcertSync/ConcertSyncClient/Source/ConcertSyncClient/Private/Replication/Manager/ReplicationManagerState_Handshaking.h
@@ -21,8 +21,8 @@ namespace UE::ConcertSyncClient::Replication
 			FJoinReplicatedSessionArgs RequestArgs,
 			TPromise<FJoinReplicatedSessionResult> JoinSessionPromise,
 			TSharedRef<IConcertClientSession> LiveSession,
-			IConcertClientReplicationBridge* ReplicationBridge,
-			FReplicationManager& Owner
+			IConcertClientReplicationBridge& ReplicationBridge UE_LIFETIMEBOUND,
+			FReplicationManager& Owner UE_LIFETIMEBOUND
 			);
 		virtual ~FReplicationManagerState_Handshaking() override;
 		
@@ -44,7 +44,7 @@ namespace UE::ConcertSyncClient::Replication
 		/** Passed to FReplicationManagerState_Handshaking */
 		TSharedRef<IConcertClientSession> LiveSession;
 		/** Passed to FReplicationManagerState_Handshaking */
-		IConcertClientReplicationBridge* ReplicationBridge;
+		IConcertClientReplicationBridge& ReplicationBridge;
 
 		//~ Begin FReplicationManagerState Interface
 		virtual void OnEnterState() override;
diff --git a/Engine/Plugins/Developer/Concert/ConcertSync/ConcertSyncClient/Source/ConcertSyncClient/Private/Replication/Processing/ClientReplicationDataCollector.cpp b/Engine/Plugins/Developer/Concert/ConcertSync/ConcertSyncClient/Source/ConcertSyncClient/Private/Replication/Processing/ClientReplicationDataCollector.cpp
--- a/Engine/Plugins/Developer/Concert/ConcertSync/ConcertSyncClient/Source/ConcertSyncClient/Private/Replication/Processing/ClientReplicationDataCollector.cpp
+++ b/Engine/Plugins/Developer/Concert/ConcertSync/ConcertSyncClient/Source/ConcertSyncClient/Private/Replication/Processing/ClientReplicationDataCollector.cpp
@@ -14,29 +14,29 @@
 namespace UE::ConcertSyncClient::Replication
 {
 	FClientReplicationDataCollector::FClientReplicationDataCollector(
-		IConcertClientReplicationBridge* InReplicationBridge,
-		TSharedRef<ConcertSyncCore::IObjectReplicationFormat> InReplicationFormat,
+		IConcertClientReplicationBridge& InReplicationBridge,
+		ConcertSyncCore::IObjectReplicationFormat& InReplicationFormat,
 		FGetClientStreams InGetStreamsDelegate,
 		const FGuid& InClientId
 		)
 		: Bridge(InReplicationBridge)
-		, ReplicationFormat(MoveTemp(InReplicationFormat))
+		, ReplicationFormat(InReplicationFormat)
 		, GetStreamsDelegate(MoveTemp(InGetStreamsDelegate))
 		, ClientId(InClientId)
 	{
 		check(GetStreamsDelegate.IsBound());
-		Bridge->OnObjectDiscovered().AddRaw(this, &FClientReplicationDataCollector::StartTrackingObject);
-		Bridge->OnObjectHidden().AddRaw(this, &FClientReplicationDataCollector::StopTrackingObject);
+		Bridge.OnObjectDiscovered().AddRaw(this, &FClientReplicationDataCollector::StartTrackingObject);
+		Bridge.OnObjectHidden().AddRaw(this, &FClientReplicationDataCollector::StopTrackingObject);
 	}
 
 	FClientReplicationDataCollector::~FClientReplicationDataCollector()
 	{
-		Bridge->OnObjectDiscovered().RemoveAll(this);
-		Bridge->OnObjectHidden().RemoveAll(this);
+		Bridge.OnObjectDiscovered().RemoveAll(this);
+		Bridge.OnObjectHidden().RemoveAll(this);
 
 		for (const TPair<FSoftObjectPath, TArray<FObjectInfo>>& Pair : ObjectsToReplicate)
 		{
-			Bridge->PopTrackedObjects({ Pair.Key });
+			Bridge.PopTrackedObjects({ Pair.Key });
 		}
 	}
 	
@@ -63,7 +63,7 @@ namespace UE::ConcertSyncClient::Replication
 		// FClientReplicationDataCollector should not push Object more than once because each push increments an internal counter
 		if (bIsNewReplicatedObject)
 		{
-			Bridge->PushTrackedObjects({ Object });
+			Bridge.PushTrackedObjects({ Object });
 		}
 	}
 
@@ -84,7 +84,7 @@ namespace UE::ConcertSyncClient::Replication
 		if (ReplicatedObjectInfo->IsEmpty())
 		{
 			ObjectsToReplicate.Remove(Object);
-			Bridge->PopTrackedObjects({ Object });
+			Bridge.PopTrackedObjects({ Object });
 		}
 	}
 
@@ -165,7 +165,7 @@ namespace UE::ConcertSyncClient::Replication
 				// Funky Unreal flows can cause the object to be renamed out from under us and replaced by a different instance (just by using UObject::Rename()).
 				// The bridge is aware of these flows and FindObjectIfAvailable will catch them.
 				// If FindObjectIfAvailable fails to resolve, OnObjectHidden() is triggered if the object was previously visible.
-				const UObject* Object = Bridge->FindObjectIfAvailable(Pair.Key);
+				const UObject* Object = Bridge.FindObjectIfAvailable(Pair.Key);
 				
 				const FSoftObjectPath ObjectPath = Object;
 				if (Object && ensureMsgf(ObjectPath == Object, TEXT("Sanity check: the bridge gave us an object with a different path!")))
@@ -200,7 +200,7 @@ namespace UE::ConcertSyncClient::Replication
 			return {};
 		}
 
-		UObject* Object = Bridge->FindObjectIfAvailable(ObjectToProcess.Object);
+		UObject* Object = Bridge.FindObjectIfAvailable(ObjectToProcess.Object);
 		// Finally... ask the bridge to resolve the object for us. The bridge has the object cached and handles the object getting renamed, etc.
 		// This should resolve. If it does not: same logic as above - either this call is invalid or ForEachPendingObject lied.
 		if (!ensure(Object))
@@ -210,7 +210,7 @@ namespace UE::ConcertSyncClient::Replication
 
 		const FObjectInfo& ObjectInfo = (*ObjectInfos)[StreamIndex];
 		ConcertSyncCore::FReplicationPropertyFilter Filter(ObjectInfo.SelectedProperties);
-		TOptional<FConcertSessionSerializedPayload> Payload = ReplicationFormat->CreateReplicationEvent(
+		TOptional<FConcertSessionSerializedPayload> Payload = ReplicationFormat.CreateReplicationEvent(
 			*Object,
 			[&Filter](const FArchiveSerializedPropertyChain* Chain, const FProperty& Property)
 			{
diff --git a/Engine/Plugins/Developer/Concert/ConcertSync/ConcertSyncClient/Source/ConcertSyncClient/Private/Replication/Processing/ClientReplicationDataCollector.h b/Engine/Plugins/Developer/Concert/ConcertSync/ConcertSyncClient/Source/ConcertSyncClient/Private/Replication/Processing/ClientReplicationDataCollector.h
--- a/Engine/Plugins/Developer/Concert/ConcertSync/ConcertSyncClient/Source/ConcertSyncClient/Private/Replication/Processing/ClientReplicationDataCollector.h
+++ b/Engine/Plugins/Developer/Concert/ConcertSync/ConcertSyncClient/Source/ConcertSyncClient/Private/Replication/Processing/ClientReplicationDataCollector.h
@@ -36,15 +36,9 @@ namespace UE::ConcertSyncClient::Replication
 
 		DECLARE_DELEGATE_RetVal(const TArray<FConcertReplicationStream>*, FGetClientStreams);
 
-		/**
-		 * @param InReplicationBridge 
-		 * @param InReplicationFormat 
-		 * @param InGetStreamsDelegate 
-		 * @param ClientId 
-		 */
 		FClientReplicationDataCollector(
-			IConcertClientReplicationBridge* InReplicationBridge,
-			TSharedRef<ConcertSyncCore::IObjectReplicationFormat> InReplicationFormat,
+			IConcertClientReplicationBridge& InReplicationBridge UE_LIFETIMEBOUND,
+			ConcertSyncCore::IObjectReplicationFormat& InReplicationFormat UE_LIFETIMEBOUND,
 			FGetClientStreams InGetStreamsDelegate,
 			const FGuid& InClientId
 			);
@@ -83,9 +77,9 @@ namespace UE::ConcertSyncClient::Replication
 	private:
 
 		/** Gets and tracks replicated objects */
-		IConcertClientReplicationBridge* Bridge;
+		IConcertClientReplicationBridge& Bridge;
 		/** Used to create the replication data sent to the server. */
-		const TSharedRef<ConcertSyncCore::IObjectReplicationFormat> ReplicationFormat;
+		ConcertSyncCore::IObjectReplicationFormat& ReplicationFormat;
 
 		/** Gets the stream of the managed client. */
 		const FGetClientStreams GetStreamsDelegate;
diff --git a/Engine/Plugins/Developer/Concert/ConcertSync/ConcertSyncClient/Source/ConcertSyncClient/Private/Replication/Processing/ClientReplicationDataQueuer.cpp b/Engine/Plugins/Developer/Concert/ConcertSync/ConcertSyncClient/Source/ConcertSyncClient/Private/Replication/Processing/ClientReplicationDataQueuer.cpp
--- a/Engine/Plugins/Developer/Concert/ConcertSync/ConcertSyncClient/Source/ConcertSyncClient/Private/Replication/Processing/ClientReplicationDataQueuer.cpp
+++ b/Engine/Plugins/Developer/Concert/ConcertSync/ConcertSyncClient/Source/ConcertSyncClient/Private/Replication/Processing/ClientReplicationDataQueuer.cpp
@@ -7,23 +7,23 @@
 namespace UE::ConcertSyncClient::Replication
 {
 	TSharedRef<FClientReplicationDataQueuer> FClientReplicationDataQueuer::Make(
-		IConcertClientReplicationBridge* ReplicationBridge,
-		TSharedRef<ConcertSyncCore::FObjectReplicationCache> InReplicationCache
+		IConcertClientReplicationBridge& ReplicationBridge,
+		ConcertSyncCore::FObjectReplicationCache& InReplicationCache
 		)
 	{
 		TSharedRef<FClientReplicationDataQueuer> Result = MakeShared<FClientReplicationDataQueuer>(ReplicationBridge);
-		Result->BindToCache(MoveTemp(InReplicationCache));
+		Result->BindToCache(InReplicationCache);
 		return Result;
 	}
 	
-	FClientReplicationDataQueuer::FClientReplicationDataQueuer(IConcertClientReplicationBridge* ReplicationBridge)
+	FClientReplicationDataQueuer::FClientReplicationDataQueuer(IConcertClientReplicationBridge& ReplicationBridge)
 		: ReplicationBridge(ReplicationBridge)
 	{}
 
 	bool FClientReplicationDataQueuer::WantsToAcceptObject(const FConcertReplicatedObjectId& Object) const
 	{
 		// The server tries not to send objects to are unavailable but it could happen the object becomes unavailable while data is underway
-		const bool bIsAvailable = ReplicationBridge->IsObjectAvailable(Object.Object);
+		const bool bIsAvailable = ReplicationBridge.IsObjectAvailable(Object.Object);
 		return bIsAvailable;
 	}
 }
diff --git a/Engine/Plugins/Developer/Concert/ConcertSync/ConcertSyncClient/Source/ConcertSyncClient/Private/Replication/Processing/ClientReplicationDataQueuer.h b/Engine/Plugins/Developer/Concert/ConcertSync/ConcertSyncClient/Source/ConcertSyncClient/Private/Replication/Processing/ClientReplicationDataQueuer.h
--- a/Engine/Plugins/Developer/Concert/ConcertSync/ConcertSyncClient/Source/ConcertSyncClient/Private/Replication/Processing/ClientReplicationDataQueuer.h
+++ b/Engine/Plugins/Developer/Concert/ConcertSync/ConcertSyncClient/Source/ConcertSyncClient/Private/Replication/Processing/ClientReplicationDataQueuer.h
@@ -16,7 +16,7 @@ namespace UE::ConcertSyncClient::Replication
 		friend class SharedPointerInternals::TIntrusiveReferenceController;
 	public:
 
-		static TSharedRef<FClientReplicationDataQueuer> Make(IConcertClientReplicationBridge* ReplicationBridge, TSharedRef<ConcertSyncCore::FObjectReplicationCache> InReplicationCache);
+		static TSharedRef<FClientReplicationDataQueuer> Make(IConcertClientReplicationBridge& ReplicationBridge UE_LIFETIMEBOUND, ConcertSyncCore::FObjectReplicationCache& InReplicationCache UE_LIFETIMEBOUND);
 		
 		//~ Begin IReplicationCacheUser Interface
 		virtual bool WantsToAcceptObject(const FConcertReplicatedObjectId& Object) const override;
@@ -24,10 +24,10 @@ namespace UE::ConcertSyncClient::Replication
 
 	private:
 		
-		FClientReplicationDataQueuer(IConcertClientReplicationBridge* ReplicationBridge);
+		FClientReplicationDataQueuer(IConcertClientReplicationBridge& ReplicationBridge UE_LIFETIMEBOUND);
 
 		/** Used to determine whether the object is currently alive. */
-		IConcertClientReplicationBridge* ReplicationBridge;
+		IConcertClientReplicationBridge& ReplicationBridge;
 	};
 }
 
diff --git a/Engine/Plugins/Developer/Concert/ConcertSync/ConcertSyncClient/Source/ConcertSyncClient/Private/Replication/Processing/ObjectReplicationApplierProcessor.cpp b/Engine/Plugins/Developer/Concert/ConcertSync/ConcertSyncClient/Source/ConcertSyncClient/Private/Replication/Processing/ObjectReplicationApplierProcessor.cpp
--- a/Engine/Plugins/Developer/Concert/ConcertSync/ConcertSyncClient/Source/ConcertSyncClient/Private/Replication/Processing/ObjectReplicationApplierProcessor.cpp
+++ b/Engine/Plugins/Developer/Concert/ConcertSync/ConcertSyncClient/Source/ConcertSyncClient/Private/Replication/Processing/ObjectReplicationApplierProcessor.cpp
@@ -11,18 +11,18 @@
 namespace UE::ConcertSyncClient::Replication
 {
 	FObjectReplicationApplierProcessor::FObjectReplicationApplierProcessor(
-		IConcertClientReplicationBridge* ReplicationBridge,
-		TSharedRef<ConcertSyncCore::IObjectReplicationFormat> ReplicationFormat,
-		TSharedRef<ConcertSyncCore::IReplicationDataSource> DataSource
+		IConcertClientReplicationBridge& ReplicationBridge,
+		ConcertSyncCore::IObjectReplicationFormat& ReplicationFormat,
+		ConcertSyncCore::IReplicationDataSource& DataSource
 		)
-		: FObjectReplicationProcessor(MoveTemp(DataSource))
+		: FObjectReplicationProcessor(DataSource)
 		, ReplicationBridge(ReplicationBridge)
-		, ReplicationFormat(MoveTemp(ReplicationFormat))
+		, ReplicationFormat(ReplicationFormat)
 	{}
 
 	void FObjectReplicationApplierProcessor::ProcessObject(const FObjectProcessArgs& Args)
 	{
-		UObject* Object = ReplicationBridge->FindObjectIfAvailable(Args.ObjectInfo.Object);
+		UObject* Object = ReplicationBridge.FindObjectIfAvailable(Args.ObjectInfo.Object);
 		if (!Object)
 		{
 			UE_LOG(LogConcert, Error, TEXT("Replication: Object %s is unavailable. The data source should not have reported it."), *Args.ObjectInfo.Object.ToString());
@@ -37,12 +37,12 @@ namespace UE::ConcertSyncClient::Replication
 			// TODO DP UE-193659: This is very hacky and leaves performance on the table... this is in case ApplyReplicationEvent updates the transform
 			if (USceneComponent* SceneComponent = Cast<USceneComponent>(Object))
 			{
-				ReplicationFormat->ApplyReplicationEvent(*Object, Payload);
+				ReplicationFormat.ApplyReplicationEvent(*Object, Payload);
 				SceneComponent->UpdateComponentToWorld();
 			}
 			else
 			{
-				ReplicationFormat->ApplyReplicationEvent(*Object, Payload);
+				ReplicationFormat.ApplyReplicationEvent(*Object, Payload);
 			}
 		});
 		// This should not happen. If it does, we're wasting  network bandwidth.
diff --git a/Engine/Plugins/Developer/Concert/ConcertSync/ConcertSyncClient/Source/ConcertSyncClient/Private/Replication/Processing/ObjectReplicationApplierProcessor.h b/Engine/Plugins/Developer/Concert/ConcertSync/ConcertSyncClient/Source/ConcertSyncClient/Private/Replication/Processing/ObjectReplicationApplierProcessor.h
--- a/Engine/Plugins/Developer/Concert/ConcertSync/ConcertSyncClient/Source/ConcertSyncClient/Private/Replication/Processing/ObjectReplicationApplierProcessor.h
+++ b/Engine/Plugins/Developer/Concert/ConcertSync/ConcertSyncClient/Source/ConcertSyncClient/Private/Replication/Processing/ObjectReplicationApplierProcessor.h
@@ -15,9 +15,9 @@ namespace UE::ConcertSyncClient::Replication
 	public:
 
 		FObjectReplicationApplierProcessor(
-			IConcertClientReplicationBridge* ReplicationBridge,
-			TSharedRef<ConcertSyncCore::IObjectReplicationFormat> ReplicationFormat,
-			TSharedRef<ConcertSyncCore::IReplicationDataSource> DataSource
+			IConcertClientReplicationBridge& ReplicationBridge UE_LIFETIMEBOUND,
+			ConcertSyncCore::IObjectReplicationFormat& ReplicationFormat UE_LIFETIMEBOUND,
+			ConcertSyncCore::IReplicationDataSource& DataSource UE_LIFETIMEBOUND
 			);
 		
 	protected:
@@ -29,10 +29,10 @@ namespace UE::ConcertSyncClient::Replication
 	private:
 
 		/** Retrieves objects for applying replication data to. */
-		IConcertClientReplicationBridge* ReplicationBridge;
+		IConcertClientReplicationBridge& ReplicationBridge;
 
 		/** Unpacks replication data and applies it to an UObject instance. */
-		TSharedRef<ConcertSyncCore::IObjectReplicationFormat> ReplicationFormat;
+		ConcertSyncCore::IObjectReplicationFormat& ReplicationFormat;
 	};
 }
 
diff --git a/Engine/Plugins/Developer/Concert/ConcertSync/ConcertSyncClient/Source/ConcertSyncClient/Private/Replication/ReplicationTestInterface.cpp b/Engine/Plugins/Developer/Concert/ConcertSync/ConcertSyncClient/Source/ConcertSyncClient/Private/Replication/ReplicationTestInterface.cpp
--- a/Engine/Plugins/Developer/Concert/ConcertSync/ConcertSyncClient/Source/ConcertSyncClient/Private/Replication/ReplicationTestInterface.cpp
+++ b/Engine/Plugins/Developer/Concert/ConcertSync/ConcertSyncClient/Source/ConcertSyncClient/Private/Replication/ReplicationTestInterface.cpp
@@ -15,7 +15,7 @@ namespace UE::ConcertSyncClient::TestInterface
 {
 	CONCERTSYNCCLIENT_API TSharedRef<IConcertClientReplicationManager> CreateClientReplicationManager(
 		TSharedRef<IConcertClientSession> InLiveSession,
-		IConcertClientReplicationBridge* InBridge
+		IConcertClientReplicationBridge& InBridge
 		)
 	{
 		const TSharedRef<Replication::FReplicationManager> Result = MakeShared<Replication::FReplicationManager>(MoveTemp(InLiveSession), InBridge);
diff --git a/Engine/Plugins/Developer/Concert/ConcertSync/ConcertSyncCore/Source/ConcertSyncCore/Private/Replication/Processing/ObjectReplicationCache.cpp b/Engine/Plugins/Developer/Concert/ConcertSync/ConcertSyncCore/Source/ConcertSyncCore/Private/Replication/Processing/ObjectReplicationCache.cpp
--- a/Engine/Plugins/Developer/Concert/ConcertSync/ConcertSyncCore/Source/ConcertSyncCore/Private/Replication/Processing/ObjectReplicationCache.cpp
+++ b/Engine/Plugins/Developer/Concert/ConcertSync/ConcertSyncCore/Source/ConcertSyncCore/Private/Replication/Processing/ObjectReplicationCache.cpp
@@ -8,8 +8,8 @@
 
 namespace UE::ConcertSyncCore
 {
-	FObjectReplicationCache::FObjectReplicationCache(TSharedRef<IObjectReplicationFormat> ReplicationFormat)
-		: ReplicationFormat(MoveTemp(ReplicationFormat))
+	FObjectReplicationCache::FObjectReplicationCache(IObjectReplicationFormat& ReplicationFormat)
+		: ReplicationFormat(ReplicationFormat)
 	{}
 
 	int32 FObjectReplicationCache::StoreUntilConsumed(const FGuid& SendingEndpointId, const FGuid& OriginStreamId, const FConcertReplication_ObjectReplicationEvent& ObjectReplicationEvent)
@@ -29,7 +29,7 @@ namespace UE::ConcertSyncCore
 				if (EventDataPin && !CombineOnceDetection.Contains(EventDataPin.Get()))
 				{
 					CombineOnceDetection.Add(EventDataPin.Get());
-					ReplicationFormat->CombineReplicationEvents(EventDataPin->SerializedPayload, ObjectReplicationEvent.SerializedPayload);
+					ReplicationFormat.CombineReplicationEvents(EventDataPin->SerializedPayload, ObjectReplicationEvent.SerializedPayload);
 				}
 			}
 		}
diff --git a/Engine/Plugins/Developer/Concert/ConcertSync/ConcertSyncCore/Source/ConcertSyncCore/Private/Replication/Processing/ObjectReplicationProcessor.cpp b/Engine/Plugins/Developer/Concert/ConcertSync/ConcertSyncCore/Source/ConcertSyncCore/Private/Replication/Processing/ObjectReplicationProcessor.cpp
--- a/Engine/Plugins/Developer/Concert/ConcertSync/ConcertSyncCore/Source/ConcertSyncCore/Private/Replication/Processing/ObjectReplicationProcessor.cpp
+++ b/Engine/Plugins/Developer/Concert/ConcertSync/ConcertSyncCore/Source/ConcertSyncCore/Private/Replication/Processing/ObjectReplicationProcessor.cpp
@@ -8,15 +8,15 @@
 namespace UE::ConcertSyncCore
 {
 	FObjectReplicationProcessor::FObjectReplicationProcessor(
-		TSharedRef<IReplicationDataSource> DataSource
+		IReplicationDataSource& DataSource
 		)
-		: DataSource(MoveTemp(DataSource))
+		: DataSource(DataSource)
 	{}
 
 	void FObjectReplicationProcessor::ProcessObjects(const FProcessObjectsParams& Params)
 	{
 		// TODO UE-190714: Respect time budget and prioritize objects
-		DataSource->ForEachPendingObject([this](const FConcertReplicatedObjectId& ObjectInfo)
+		DataSource.ForEachPendingObject([this](const FConcertReplicatedObjectId& ObjectInfo)
 		{
 			ProcessObject({ ObjectInfo });
 		});
diff --git a/Engine/Plugins/Developer/Concert/ConcertSync/ConcertSyncCore/Source/ConcertSyncCore/Private/Replication/Processing/ObjectReplicationReceiver.cpp b/Engine/Plugins/Developer/Concert/ConcertSync/ConcertSyncCore/Source/ConcertSyncCore/Private/Replication/Processing/ObjectReplicationReceiver.cpp
--- a/Engine/Plugins/Developer/Concert/ConcertSync/ConcertSyncCore/Source/ConcertSyncCore/Private/Replication/Processing/ObjectReplicationReceiver.cpp
+++ b/Engine/Plugins/Developer/Concert/ConcertSync/ConcertSyncCore/Source/ConcertSyncCore/Private/Replication/Processing/ObjectReplicationReceiver.cpp
@@ -13,16 +13,16 @@ namespace UE::ConcertSyncCore
 {
 	static TAutoConsoleVariable<bool> CVarLogReceivedObjects(TEXT("Concert.Replication.LogReceivedObjects"), false, TEXT("Enable Concert logging for received replicated objects."));
 	
-	FObjectReplicationReceiver::FObjectReplicationReceiver(TSharedRef<IConcertSession> Session, TSharedRef<FObjectReplicationCache> ReplicationCache)
-		: Session(MoveTemp(Session))
-		, ReplicationCache(MoveTemp(ReplicationCache))
+	FObjectReplicationReceiver::FObjectReplicationReceiver(IConcertSession& Session, FObjectReplicationCache& ReplicationCache)
+		: Session(Session)
+		, ReplicationCache(ReplicationCache)
 	{
-		Session->RegisterCustomEventHandler<FConcertReplication_BatchReplicationEvent>(this, &FObjectReplicationReceiver::HandleBatchReplicationEvent);
+		Session.RegisterCustomEventHandler<FConcertReplication_BatchReplicationEvent>(this, &FObjectReplicationReceiver::HandleBatchReplicationEvent);
 	}
 
 	FObjectReplicationReceiver::~FObjectReplicationReceiver()
 	{
-		Session->UnregisterCustomEventHandler<FConcertReplication_BatchReplicationEvent>(this);
+		Session.UnregisterCustomEventHandler<FConcertReplication_BatchReplicationEvent>(this);
 	}
 
 	void FObjectReplicationReceiver::HandleBatchReplicationEvent(const FConcertSessionContext& SessionContext, const FConcertReplication_BatchReplicationEvent& Event)
@@ -42,7 +42,7 @@ namespace UE::ConcertSyncCore
 			{
 				if (ShouldAcceptObject(SessionContext, StreamEvent, ObjectEvent))
 				{
-					const int32 NumAccepted = ReplicationCache->StoreUntilConsumed(SessionContext.SourceEndpointId, StreamEvent.StreamId, ObjectEvent);
+					const int32 NumAccepted = ReplicationCache.StoreUntilConsumed(SessionContext.SourceEndpointId, StreamEvent.StreamId, ObjectEvent);
 					NumCacheUsages += NumAccepted;
 					NumOfAcceptedObjectChanges += NumAccepted == 0 ? 0 : 1;
 				}
diff --git a/Engine/Plugins/Developer/Concert/ConcertSync/ConcertSyncCore/Source/ConcertSyncCore/Private/Replication/Processing/ObjectReplicationSender.cpp b/Engine/Plugins/Developer/Concert/ConcertSync/ConcertSyncCore/Source/ConcertSyncCore/Private/Replication/Processing/ObjectReplicationSender.cpp
--- a/Engine/Plugins/Developer/Concert/ConcertSync/ConcertSyncCore/Source/ConcertSyncCore/Private/Replication/Processing/ObjectReplicationSender.cpp
+++ b/Engine/Plugins/Developer/Concert/ConcertSync/ConcertSyncCore/Source/ConcertSyncCore/Private/Replication/Processing/ObjectReplicationSender.cpp
@@ -15,12 +15,12 @@ namespace UE::ConcertSyncCore
 	
 	FObjectReplicationSender::FObjectReplicationSender(
 		const FGuid& TargetEndpointId,
-		TSharedRef<IConcertSession> Session,
-		TSharedRef<IReplicationDataSource> DataSource
+		IConcertSession& Session,
+		IReplicationDataSource& DataSource
 		)
-		: FObjectReplicationProcessor(MoveTemp(DataSource))
+		: FObjectReplicationProcessor(DataSource)
 		, TargetEndpointId(TargetEndpointId)
-		, Session(MoveTemp(Session))
+		, Session(Session)
 	{}
 
 	void FObjectReplicationSender::ProcessObjects(const FProcessObjectsParams& Params)
@@ -36,7 +36,7 @@ namespace UE::ConcertSyncCore
 				*TargetEndpointId.ToString()
 				);
 			
-			Session->SendCustomEvent(EventToSend, TargetEndpointId,
+			Session.SendCustomEvent(EventToSend, TargetEndpointId,
 				// Replication is always unreliable - if it fails to deliver we'll send updated data soon again
 				// TODO: In regular intervals send CRC values to detect that a change is missing
 				EConcertMessageFlags::None
diff --git a/Engine/Plugins/Developer/Concert/ConcertSync/ConcertSyncCore/Source/ConcertSyncCore/Private/Replication/Processing/ReplicationDataQueuer.cpp b/Engine/Plugins/Developer/Concert/ConcertSync/ConcertSyncCore/Source/ConcertSyncCore/Private/Replication/Processing/ReplicationDataQueuer.cpp
--- a/Engine/Plugins/Developer/Concert/ConcertSync/ConcertSyncCore/Source/ConcertSyncCore/Private/Replication/Processing/ReplicationDataQueuer.cpp
+++ b/Engine/Plugins/Developer/Concert/ConcertSync/ConcertSyncCore/Source/ConcertSyncCore/Private/Replication/Processing/ReplicationDataQueuer.cpp
@@ -6,14 +6,6 @@
 
 namespace UE::ConcertSyncCore
 {
-	FReplicationDataQueuer::~FReplicationDataQueuer()
-	{
-		if (ensure(ReplicationCache))
-		{
-			ReplicationCache->UnregisterDataCacheUser(AsShared());
-		}
-	}
-
 	void FReplicationDataQueuer::ForEachPendingObject(TFunctionRef<void(const FConcertReplicatedObjectId&)> ProcessItemFunc) const
 	{
 		for (auto It = PendingEvents.CreateConstIterator(); It; ++It)
@@ -50,12 +42,9 @@ namespace UE::ConcertSyncCore
 		PendingEvents.Add(Object, MoveTemp(Data));
 	}
 	
-	void FReplicationDataQueuer::BindToCache(TSharedRef<FObjectReplicationCache> InReplicationCache)
+	void FReplicationDataQueuer::BindToCache(FObjectReplicationCache& InReplicationCache)
 	{
-		if (ensure(!ReplicationCache))
-		{
-			ReplicationCache = MoveTemp(InReplicationCache);
-			ReplicationCache->RegisterDataCacheUser(AsShared());
-		}
+		ReplicationCache = &InReplicationCache;
+		ReplicationCache->RegisterDataCacheUser(AsShared());
 	}
 }
diff --git a/Engine/Plugins/Developer/Concert/ConcertSync/ConcertSyncCore/Source/ConcertSyncCore/Public/Replication/Processing/ObjectReplicationCache.h b/Engine/Plugins/Developer/Concert/ConcertSync/ConcertSyncCore/Source/ConcertSyncCore/Public/Replication/Processing/ObjectReplicationCache.h
--- a/Engine/Plugins/Developer/Concert/ConcertSync/ConcertSyncCore/Source/ConcertSyncCore/Public/Replication/Processing/ObjectReplicationCache.h
+++ b/Engine/Plugins/Developer/Concert/ConcertSync/ConcertSyncCore/Source/ConcertSyncCore/Public/Replication/Processing/ObjectReplicationCache.h
@@ -54,7 +54,7 @@ namespace UE::ConcertSyncCore
 	{
 	public:
 
-		FObjectReplicationCache(TSharedRef<IObjectReplicationFormat> ReplicationFormat);
+		FObjectReplicationCache(IObjectReplicationFormat& ReplicationFormat UE_LIFETIMEBOUND);
 
 		/**
 		 * Called when new data is received for an object and shares it with any IObjectCacheUser that is possibly interested in it.
@@ -72,7 +72,7 @@ namespace UE::ConcertSyncCore
 	private:
 
 		/** Used for combining events to save network bandwidth. */
-		TSharedRef<IObjectReplicationFormat> ReplicationFormat;
+		IObjectReplicationFormat& ReplicationFormat;
 
 		/** Everyone who registered for receiving data. */
 		TArray<TSharedRef<IReplicationCacheUser>> CacheUsers;
diff --git a/Engine/Plugins/Developer/Concert/ConcertSync/ConcertSyncCore/Source/ConcertSyncCore/Public/Replication/Processing/ObjectReplicationProcessor.h b/Engine/Plugins/Developer/Concert/ConcertSync/ConcertSyncCore/Source/ConcertSyncCore/Public/Replication/Processing/ObjectReplicationProcessor.h
--- a/Engine/Plugins/Developer/Concert/ConcertSync/ConcertSyncCore/Source/ConcertSyncCore/Public/Replication/Processing/ObjectReplicationProcessor.h
+++ b/Engine/Plugins/Developer/Concert/ConcertSync/ConcertSyncCore/Source/ConcertSyncCore/Public/Replication/Processing/ObjectReplicationProcessor.h
@@ -58,7 +58,7 @@ namespace UE::ConcertSyncCore
 		/**
 		 * @param DataSource Source of the data that is to be sent
 		 */
-		FObjectReplicationProcessor(TSharedRef<IReplicationDataSource> DataSource);
+		FObjectReplicationProcessor(IReplicationDataSource& DataSource UE_LIFETIMEBOUND);
 		virtual ~FObjectReplicationProcessor() = default;
 		
 		virtual void ProcessObjects(const FProcessObjectsParams& Params);
@@ -71,7 +71,7 @@ namespace UE::ConcertSyncCore
 			FConcertReplicatedObjectId ObjectInfo;
 		};
 
-		FORCEINLINE IReplicationDataSource& GetDataSource() const { return DataSource.Get(); }
+		FORCEINLINE IReplicationDataSource& GetDataSource() const { return DataSource; }
 		
 		/** Processes the object. */
 		virtual void ProcessObject(const FObjectProcessArgs& Args) = 0;
@@ -79,6 +79,6 @@ namespace UE::ConcertSyncCore
 	private:
 		
 		/** Abstracts where replication data comes from: could be generated (clients) or received (server or client) */
-		const TSharedRef<IReplicationDataSource> DataSource;
+		IReplicationDataSource& DataSource;
 	};
 }
diff --git a/Engine/Plugins/Developer/Concert/ConcertSync/ConcertSyncCore/Source/ConcertSyncCore/Public/Replication/Processing/ObjectReplicationReceiver.h b/Engine/Plugins/Developer/Concert/ConcertSync/ConcertSyncCore/Source/ConcertSyncCore/Public/Replication/Processing/ObjectReplicationReceiver.h
--- a/Engine/Plugins/Developer/Concert/ConcertSync/ConcertSyncCore/Source/ConcertSyncCore/Public/Replication/Processing/ObjectReplicationReceiver.h
+++ b/Engine/Plugins/Developer/Concert/ConcertSync/ConcertSyncCore/Source/ConcertSyncCore/Public/Replication/Processing/ObjectReplicationReceiver.h
@@ -2,7 +2,7 @@
 
 #pragma once
 
-#include "Templates/SharedPointer.h"
+#include "HAL/Platform.h"
 
 class IConcertSession;
 struct FConcertReplication_BatchReplicationEvent;
@@ -19,7 +19,7 @@ namespace UE::ConcertSyncCore
 	{
 	public:
 
-		FObjectReplicationReceiver(TSharedRef<IConcertSession> Session, TSharedRef<FObjectReplicationCache> ReplicationCache);
+		FObjectReplicationReceiver(IConcertSession& Session UE_LIFETIMEBOUND, FObjectReplicationCache& ReplicationCache UE_LIFETIMEBOUND);
 		virtual ~FObjectReplicationReceiver();
 
 	protected:
@@ -30,9 +30,9 @@ namespace UE::ConcertSyncCore
 	private:
 
 		/** The session that is being received on. */
-		TSharedRef<IConcertSession> Session;
+		IConcertSession& Session;
 		/** Where received data is stored. */
-		TSharedRef<FObjectReplicationCache> ReplicationCache;
+		FObjectReplicationCache& ReplicationCache;
 
 		void HandleBatchReplicationEvent(const FConcertSessionContext& SessionContext, const FConcertReplication_BatchReplicationEvent& Event);
 	};
diff --git a/Engine/Plugins/Developer/Concert/ConcertSync/ConcertSyncCore/Source/ConcertSyncCore/Public/Replication/Processing/ObjectReplicationSender.h b/Engine/Plugins/Developer/Concert/ConcertSync/ConcertSyncCore/Source/ConcertSyncCore/Public/Replication/Processing/ObjectReplicationSender.h
--- a/Engine/Plugins/Developer/Concert/ConcertSync/ConcertSyncCore/Source/ConcertSyncCore/Public/Replication/Processing/ObjectReplicationSender.h
+++ b/Engine/Plugins/Developer/Concert/ConcertSync/ConcertSyncCore/Source/ConcertSyncCore/Public/Replication/Processing/ObjectReplicationSender.h
@@ -23,8 +23,8 @@ namespace UE::ConcertSyncCore
 		 */
 		FObjectReplicationSender(
 			const FGuid& TargetEndpointId,
-			TSharedRef<IConcertSession> Session,
-			TSharedRef<IReplicationDataSource> DataSource
+			IConcertSession& Session UE_LIFETIMEBOUND, 
+			IReplicationDataSource& DataSource UE_LIFETIMEBOUND
 			);
 		
 		//~ Begin FObjectReplicationProcessor Interface
@@ -43,7 +43,7 @@ namespace UE::ConcertSyncCore
 		const FGuid TargetEndpointId;
 		
 		/** The session through which replication messages are sent. */
-		const TSharedRef<IConcertSession> Session;
+		IConcertSession& Session;
 
 		/** This event is filled in ProcessObjects and finally sent to TargetEndpointId. */
 		FConcertReplication_BatchReplicationEvent EventToSend;
diff --git a/Engine/Plugins/Developer/Concert/ConcertSync/ConcertSyncCore/Source/ConcertSyncCore/Public/Replication/Processing/ReplicationDataQueuer.h b/Engine/Plugins/Developer/Concert/ConcertSync/ConcertSyncCore/Source/ConcertSyncCore/Public/Replication/Processing/ReplicationDataQueuer.h
--- a/Engine/Plugins/Developer/Concert/ConcertSync/ConcertSyncCore/Source/ConcertSyncCore/Public/Replication/Processing/ReplicationDataQueuer.h
+++ b/Engine/Plugins/Developer/Concert/ConcertSync/ConcertSyncCore/Source/ConcertSyncCore/Public/Replication/Processing/ReplicationDataQueuer.h
@@ -24,8 +24,6 @@ namespace UE::ConcertSyncCore
 	{
 	public:
 
-		virtual ~FReplicationDataQueuer() override;
-
 		//~ Begin IReplicationDataSource Interface
 		virtual void ForEachPendingObject(TFunctionRef<void(const FConcertReplicatedObjectId&)> ProcessItemFunc) const override;
 		virtual int32 NumObjects() const override;
@@ -39,7 +37,7 @@ namespace UE::ConcertSyncCore
 	protected:
 
 		/** Called by subclass factory functions. */
-		void BindToCache(TSharedRef<FObjectReplicationCache> InReplicationCache);
+		void BindToCache(FObjectReplicationCache& InReplicationCache UE_LIFETIMEBOUND);
 		
 	private:
 
@@ -47,6 +45,6 @@ namespace UE::ConcertSyncCore
 		TMap<FConcertReplicatedObjectId, TSharedPtr<const FConcertReplication_ObjectReplicationEvent>> PendingEvents;
 
 		/** Provides us with replication events and shares them effectively. */
-		TSharedPtr<FObjectReplicationCache> ReplicationCache;
+		FObjectReplicationCache* ReplicationCache = nullptr;
 	};
 }
diff --git a/Engine/Plugins/Developer/Concert/ConcertSync/ConcertSyncServer/Source/ConcertSyncServer/Private/Replication/ConcertReplicationClient.cpp b/Engine/Plugins/Developer/Concert/ConcertSync/ConcertSyncServer/Source/ConcertSyncServer/Private/Replication/ConcertReplicationClient.cpp
--- a/Engine/Plugins/Developer/Concert/ConcertSync/ConcertSyncServer/Source/ConcertSyncServer/Private/Replication/ConcertReplicationClient.cpp
+++ b/Engine/Plugins/Developer/Concert/ConcertSync/ConcertSyncServer/Source/ConcertSyncServer/Private/Replication/ConcertReplicationClient.cpp
@@ -11,16 +11,16 @@ namespace UE::ConcertSyncServer::Replication
 	FConcertReplicationClient::FConcertReplicationClient(
 		TArray<FConcertReplicationStream> StreamDescriptions,
 		const FGuid& ClientEndpointId,
-		TSharedRef<IConcertSession> Session,
-		TSharedRef<ConcertSyncCore::FObjectReplicationCache> ReplicationCache,
+		IConcertSession& Session,
+		ConcertSyncCore::FObjectReplicationCache& ReplicationCache,
 		ConcertSyncCore::FGetObjectFrequencySettings GetObjectFrequencySettingsDelegate
 	)
 		: StreamDescriptions(MoveTemp(StreamDescriptions))
 		, ClientEndpointId(ClientEndpointId)
-		, EventQueue(FServerReplicationDataQueuer::Make(ClientEndpointId, MoveTemp(ReplicationCache)))
+		, EventQueue(FServerReplicationDataQueuer::Make(ClientEndpointId, ReplicationCache))
 		, DataRelay(
 			MoveTemp(GetObjectFrequencySettingsDelegate),
-			ClientEndpointId, MoveTemp(Session), EventQueue
+			ClientEndpointId, Session, *EventQueue
 			)
 	{}
 
diff --git a/Engine/Plugins/Developer/Concert/ConcertSync/ConcertSyncServer/Source/ConcertSyncServer/Private/Replication/ConcertReplicationClient.h b/Engine/Plugins/Developer/Concert/ConcertSync/ConcertSyncServer/Source/ConcertSyncServer/Private/Replication/ConcertReplicationClient.h
--- a/Engine/Plugins/Developer/Concert/ConcertSync/ConcertSyncServer/Source/ConcertSyncServer/Private/Replication/ConcertReplicationClient.h
+++ b/Engine/Plugins/Developer/Concert/ConcertSync/ConcertSyncServer/Source/ConcertSyncServer/Private/Replication/ConcertReplicationClient.h
@@ -29,8 +29,8 @@ namespace UE::ConcertSyncServer::Replication
 		FConcertReplicationClient(
 			TArray<FConcertReplicationStream> StreamDescriptions,
 			const FGuid& ClientEndpointId,
-			TSharedRef<IConcertSession> Session,
-			TSharedRef<ConcertSyncCore::FObjectReplicationCache> ReplicationCache,
+			IConcertSession& Session UE_LIFETIMEBOUND,
+			ConcertSyncCore::FObjectReplicationCache& ReplicationCache UE_LIFETIMEBOUND,
 			ConcertSyncCore::FGetObjectFrequencySettings GetObjectFrequencySettingsDelegate
 			);
 
diff --git a/Engine/Plugins/Developer/Concert/ConcertSync/ConcertSyncServer/Source/ConcertSyncServer/Private/Replication/ConcertServerReplicationManager.cpp b/Engine/Plugins/Developer/Concert/ConcertSync/ConcertSyncServer/Source/ConcertSyncServer/Private/Replication/ConcertServerReplicationManager.cpp
--- a/Engine/Plugins/Developer/Concert/ConcertSync/ConcertSyncServer/Source/ConcertSyncServer/Private/Replication/ConcertServerReplicationManager.cpp
+++ b/Engine/Plugins/Developer/Concert/ConcertSync/ConcertSyncServer/Source/ConcertSyncServer/Private/Replication/ConcertServerReplicationManager.cpp
@@ -25,10 +25,10 @@ namespace UE::ConcertSyncServer::Replication
 	
 	FConcertServerReplicationManager::FConcertServerReplicationManager(TSharedRef<IConcertServerSession> InLiveSession)
 		: Session(MoveTemp(InLiveSession))
-		, ReplicationFormat(MakeShared<ConcertSyncCore::FFullObjectFormat>())
-		, AuthorityManager(MakeShared<FAuthorityManager>(*this, Session))
-		, ReplicationCache(MakeShared<ConcertSyncCore::FObjectReplicationCache>(ReplicationFormat))
-		, ReplicationDataReceiver(AuthorityManager, Session, ReplicationCache)
+		, ReplicationFormat(MakeUnique<ConcertSyncCore::FFullObjectFormat>())
+		, AuthorityManager(*this, Session)
+		, ReplicationCache(MakeShared<ConcertSyncCore::FObjectReplicationCache>(*ReplicationFormat))
+		, ReplicationDataReceiver(AuthorityManager, *Session, *ReplicationCache)
 	{
 		Session->RegisterCustomRequestHandler<FConcertReplication_Join_Request, FConcertReplication_Join_Response>(this, &FConcertServerReplicationManager::HandleJoinReplicationSessionRequest);
 		Session->RegisterCustomRequestHandler<FConcertReplication_QueryReplicationInfo_Request, FConcertReplication_QueryReplicationInfo_Response>(this, &FConcertServerReplicationManager::HandleQueryReplicationInfoRequest);
@@ -127,8 +127,8 @@ namespace UE::ConcertSyncServer::Replication
 			MakeUnique<FConcertReplicationClient>(
 				MoveTemp(StreamDescriptions),
 				ClientId,
-				Session,
-				ReplicationCache,
+				*Session,
+				*ReplicationCache,
 				ConcertSyncCore::FGetObjectFrequencySettings::CreateRaw(this, &FConcertServerReplicationManager::GetObjectFrequencySettings)
 			)
 		);
@@ -205,7 +205,7 @@ namespace UE::ConcertSyncServer::Replication
 				ObjectInfo.Object = Pair.Key;
 				ObjectInfo.StreamId = StreamId;
 				
-				if (AuthorityManager->HasAuthorityToChange(ObjectInfo))
+				if (AuthorityManager.HasAuthorityToChange(ObjectInfo))
 				{
 					Info.AuthoredObjects.Add(Pair.Key);
 				}
@@ -228,7 +228,7 @@ namespace UE::ConcertSyncServer::Replication
 		UE_LOG(LogConcert, Log, TEXT("Received replication leave request from endpoint %s"), *ClientEndpointId.ToString());
 		
 		Clients.Remove(ClientEndpointId);
-		AuthorityManager->OnClientLeft(ClientEndpointId);
+		AuthorityManager.OnClientLeft(ClientEndpointId);
 	}
 
 	void FConcertServerReplicationManager::OnConnectionChanged(IConcertServerSession& ConcertServerSession, EConcertClientStatus ConcertClientStatus, const FConcertSessionClientInfo& ClientInfo)
@@ -237,7 +237,7 @@ namespace UE::ConcertSyncServer::Replication
 		if (ConcertClientStatus == EConcertClientStatus::Disconnected)
 		{
 			Clients.Remove(ClientEndpointId);
-			AuthorityManager->OnClientLeft(ClientEndpointId);
+			AuthorityManager.OnClientLeft(ClientEndpointId);
 		}
 	}
 
diff --git a/Engine/Plugins/Developer/Concert/ConcertSync/ConcertSyncServer/Source/ConcertSyncServer/Private/Replication/ConcertServerReplicationManager.h b/Engine/Plugins/Developer/Concert/ConcertSync/ConcertSyncServer/Source/ConcertSyncServer/Private/Replication/ConcertServerReplicationManager.h
--- a/Engine/Plugins/Developer/Concert/ConcertSync/ConcertSyncServer/Source/ConcertSyncServer/Private/Replication/ConcertServerReplicationManager.h
+++ b/Engine/Plugins/Developer/Concert/ConcertSync/ConcertSyncServer/Source/ConcertSyncServer/Private/Replication/ConcertServerReplicationManager.h
@@ -5,9 +5,11 @@
 #include "AuthorityManager.h"
 #include "ConcertMessages.h"
 #include "ConcertReplicationClient.h"
-#include "Processing/ServerObjectReplicationReceiver.h"
+#include "Replication/Formats/IObjectReplicationFormat.h"
 #include "Replication/IConcertServerReplicationManager.h"
 #include "Replication/Messages/Handshake.h"
+#include "Replication/Processing/ObjectReplicationCache.h"
+#include "Replication/Processing/ServerObjectReplicationReceiver.h"
 
 #include "Templates/SharedPointer.h"
 #include "Templates/Tuple.h"
@@ -49,7 +51,7 @@ namespace UE::ConcertSyncServer::Replication
 		explicit FConcertServerReplicationManager(TSharedRef<IConcertServerSession> InLiveSession);
 		virtual ~FConcertServerReplicationManager() override;
 
-		const FAuthorityManager& GetAuthorityManager() const { return AuthorityManager.Get(); }
+		const FAuthorityManager& GetAuthorityManager() const { return AuthorityManager; }
 
 		//~ Begin IAuthorityManagerGetters Interface
 		virtual void ForEachStream(const FGuid& ClientEndpointId, TFunctionRef<EBreakBehavior(const FConcertReplicationStream& Stream)> Callback) const override;
@@ -62,13 +64,13 @@ namespace UE::ConcertSyncServer::Replication
 		TSharedRef<IConcertServerSession> Session;
 		
 		/** Responsible for analysing received replication data. */
-		TSharedRef<ConcertSyncCore::IObjectReplicationFormat> ReplicationFormat;
+		TUniquePtr<ConcertSyncCore::IObjectReplicationFormat> ReplicationFormat;
 
 		/** Responds to client requests to changing authority and can be asked whether an object change is valid to take place. */
-		TSharedRef<FAuthorityManager> AuthorityManager;
+		FAuthorityManager AuthorityManager;
 		
 		/** Received replication events are put into the ReplicationCache. The cache is used to relay data to clients latently. */
-		TSharedRef<ConcertSyncCore::FObjectReplicationCache> ReplicationCache;
+		const TSharedRef<ConcertSyncCore::FObjectReplicationCache> ReplicationCache;
 		/** Receives replication events from all endpoints. */
 		FServerObjectReplicationReceiver ReplicationDataReceiver;
 
diff --git a/Engine/Plugins/Developer/Concert/ConcertSync/ConcertSyncServer/Source/ConcertSyncServer/Private/Replication/ConcertServerReplicationManager_ChangeStream.cpp b/Engine/Plugins/Developer/Concert/ConcertSync/ConcertSyncServer/Source/ConcertSyncServer/Private/Replication/ConcertServerReplicationManager_ChangeStream.cpp
--- a/Engine/Plugins/Developer/Concert/ConcertSync/ConcertSyncServer/Source/ConcertSyncServer/Private/Replication/ConcertServerReplicationManager_ChangeStream.cpp
+++ b/Engine/Plugins/Developer/Concert/ConcertSync/ConcertSyncServer/Source/ConcertSyncServer/Private/Replication/ConcertServerReplicationManager_ChangeStream.cpp
@@ -155,13 +155,13 @@ namespace UE::ConcertSyncServer::Replication
 		
 		const FGuid SendingClientId = ConcertSessionContext.SourceEndpointId;
 		const TUniquePtr<FConcertReplicationClient>* SendingClient = Clients.Find(SendingClientId);
-		if (SendingClient && Private::ShouldAcceptRequest(Request, *SendingClient->Get(), AuthorityManager.Get(), Response))
+		if (SendingClient && Private::ShouldAcceptRequest(Request, *SendingClient->Get(), AuthorityManager, Response))
 		{
 			// If the client had authority over any objects that were removed by this request, authority must be cleaned up
 			ConcertSyncCore::Replication::ChangeStreamUtils::ForEachObjectLosingAuthority(Request, SendingClient->Get()->GetStreamDescriptions(),
 				[this, &SendingClientId](const FConcertObjectInStreamID& RemovedObject)
 				{
-					AuthorityManager->RemoveAuthority({ RemovedObject, SendingClientId});
+					AuthorityManager.RemoveAuthority({ RemovedObject, SendingClientId});
 					return EBreakBehavior::Continue;
 				});
 			
diff --git a/Engine/Plugins/Developer/Concert/ConcertSync/ConcertSyncServer/Source/ConcertSyncServer/Private/Replication/Processing/ServerObjectReplicationReceiver.cpp b/Engine/Plugins/Developer/Concert/ConcertSync/ConcertSyncServer/Source/ConcertSyncServer/Private/Replication/Processing/ServerObjectReplicationReceiver.cpp
--- a/Engine/Plugins/Developer/Concert/ConcertSync/ConcertSyncServer/Source/ConcertSyncServer/Private/Replication/Processing/ServerObjectReplicationReceiver.cpp
+++ b/Engine/Plugins/Developer/Concert/ConcertSync/ConcertSyncServer/Source/ConcertSyncServer/Private/Replication/Processing/ServerObjectReplicationReceiver.cpp
@@ -10,12 +10,12 @@
 namespace UE::ConcertSyncServer::Replication
 {
 	FServerObjectReplicationReceiver::FServerObjectReplicationReceiver(
-		TSharedRef<FAuthorityManager> AuthorityManager,
-		TSharedRef<IConcertSession> Session,
-		TSharedRef<ConcertSyncCore::FObjectReplicationCache> ReplicationCache
+		const FAuthorityManager& AuthorityManager,
+		IConcertSession& Session,
+		ConcertSyncCore::FObjectReplicationCache& ReplicationCache
 		)
-		: FObjectReplicationReceiver(MoveTemp(Session), MoveTemp(ReplicationCache))
-		, AuthorityManager(MoveTemp(AuthorityManager))
+		: FObjectReplicationReceiver(Session, ReplicationCache)
+		, AuthorityManager(AuthorityManager)
 	{}
 
 	bool FServerObjectReplicationReceiver::ShouldAcceptObject(
@@ -25,6 +25,6 @@ namespace UE::ConcertSyncServer::Replication
 		) const
 	{
 		const FConcertReplicatedObjectId ReplicatedObjectInfo { { StreamEvent.StreamId, ObjectEvent.ReplicatedObject }, SessionContext.SourceEndpointId };
-		return AuthorityManager->HasAuthorityToChange(ReplicatedObjectInfo);
+		return AuthorityManager.HasAuthorityToChange(ReplicatedObjectInfo);
 	}
 }
diff --git a/Engine/Plugins/Developer/Concert/ConcertSync/ConcertSyncServer/Source/ConcertSyncServer/Private/Replication/Processing/ServerObjectReplicationReceiver.h b/Engine/Plugins/Developer/Concert/ConcertSync/ConcertSyncServer/Source/ConcertSyncServer/Private/Replication/Processing/ServerObjectReplicationReceiver.h
--- a/Engine/Plugins/Developer/Concert/ConcertSync/ConcertSyncServer/Source/ConcertSyncServer/Private/Replication/Processing/ServerObjectReplicationReceiver.h
+++ b/Engine/Plugins/Developer/Concert/ConcertSync/ConcertSyncServer/Source/ConcertSyncServer/Private/Replication/Processing/ServerObjectReplicationReceiver.h
@@ -3,7 +3,6 @@
 #pragma once
 
 #include "Replication/Processing/ObjectReplicationReceiver.h"
-#include "Templates/SharedPointer.h"
 
 namespace UE::ConcertSyncServer::Replication
 {
@@ -14,7 +13,11 @@ namespace UE::ConcertSyncServer::Replication
 	{
 	public:
 
-		FServerObjectReplicationReceiver(TSharedRef<FAuthorityManager> AuthorityManager, TSharedRef<IConcertSession> Session, TSharedRef<ConcertSyncCore::FObjectReplicationCache> ReplicationCache);
+		FServerObjectReplicationReceiver(
+			const FAuthorityManager& AuthorityManager UE_LIFETIMEBOUND,
+			IConcertSession& Session UE_LIFETIMEBOUND,
+			ConcertSyncCore::FObjectReplicationCache& ReplicationCache UE_LIFETIMEBOUND
+			);
 
 	protected:
 
@@ -25,7 +28,7 @@ namespace UE::ConcertSyncServer::Replication
 	private:
 
 		/** Used to determine whether a client has authority over objects. */
-		TSharedRef<FAuthorityManager> AuthorityManager;
+		const FAuthorityManager& AuthorityManager;
 	};
 }
 
diff --git a/Engine/Plugins/Developer/Concert/ConcertSync/ConcertSyncServer/Source/ConcertSyncServer/Private/Replication/Processing/ServerReplicationDataQueuer.cpp b/Engine/Plugins/Developer/Concert/ConcertSync/ConcertSyncServer/Source/ConcertSyncServer/Private/Replication/Processing/ServerReplicationDataQueuer.cpp
--- a/Engine/Plugins/Developer/Concert/ConcertSync/ConcertSyncServer/Source/ConcertSyncServer/Private/Replication/Processing/ServerReplicationDataQueuer.cpp
+++ b/Engine/Plugins/Developer/Concert/ConcertSync/ConcertSyncServer/Source/ConcertSyncServer/Private/Replication/Processing/ServerReplicationDataQueuer.cpp
@@ -8,11 +8,11 @@ namespace UE::ConcertSyncServer::Replication
 {
 	TSharedRef<FServerReplicationDataQueuer> FServerReplicationDataQueuer::Make(
 		const FGuid& OwningClientEndpointId,
-		TSharedRef<ConcertSyncCore::FObjectReplicationCache> InReplicationCache
+		ConcertSyncCore::FObjectReplicationCache& InReplicationCache
 		)
 	{
 		TSharedRef<FServerReplicationDataQueuer> Result = MakeShared<FServerReplicationDataQueuer>(OwningClientEndpointId);
-		Result->BindToCache(MoveTemp(InReplicationCache));
+		Result->BindToCache(InReplicationCache);
 		return Result;
 	}
 
diff --git a/Engine/Plugins/Developer/Concert/ConcertSync/ConcertSyncServer/Source/ConcertSyncServer/Private/Replication/Processing/ServerReplicationDataQueuer.h b/Engine/Plugins/Developer/Concert/ConcertSync/ConcertSyncServer/Source/ConcertSyncServer/Private/Replication/Processing/ServerReplicationDataQueuer.h
--- a/Engine/Plugins/Developer/Concert/ConcertSync/ConcertSyncServer/Source/ConcertSyncServer/Private/Replication/Processing/ServerReplicationDataQueuer.h
+++ b/Engine/Plugins/Developer/Concert/ConcertSync/ConcertSyncServer/Source/ConcertSyncServer/Private/Replication/Processing/ServerReplicationDataQueuer.h
@@ -21,7 +21,7 @@ namespace UE::ConcertSyncServer::Replication
 		friend class SharedPointerInternals::TIntrusiveReferenceController;
 	public:
 		
-		static TSharedRef<FServerReplicationDataQueuer> Make(const FGuid& OwningClientEndpointId, TSharedRef<ConcertSyncCore::FObjectReplicationCache> InReplicationCache);
+		static TSharedRef<FServerReplicationDataQueuer> Make(const FGuid& OwningClientEndpointId, ConcertSyncCore::FObjectReplicationCache& InReplicationCache UE_LIFETIMEBOUND);
 		
 		//~ Begin IReplicationCacheUser Interface
 		virtual bool WantsToAcceptObject(const FConcertReplicatedObjectId& Object) const override;
diff --git a/Engine/Plugins/Developer/Concert/ConcertSync/ConcertSyncTest/Source/ConcertSyncTest/Private/Replication/Components/ReplicationCacheTests.cpp b/Engine/Plugins/Developer/Concert/ConcertSync/ConcertSyncTest/Source/ConcertSyncTest/Private/Replication/Components/ReplicationCacheTests.cpp
--- a/Engine/Plugins/Developer/Concert/ConcertSync/ConcertSyncTest/Source/ConcertSyncTest/Private/Replication/Components/ReplicationCacheTests.cpp
+++ b/Engine/Plugins/Developer/Concert/ConcertSync/ConcertSyncTest/Source/ConcertSyncTest/Private/Replication/Components/ReplicationCacheTests.cpp
@@ -125,8 +125,9 @@ namespace UE::ConcertSyncTests
 		const FSoftObjectPath ObjectPath(TEXT("/Game/World.World:PersistentLevel.StaticMeshActor0"));
 		const FGuid DummySendingClientId = FGuid::NewGuid();
 		const FConcertReplicatedObjectId ObjectID{ { StreamId, ObjectPath }, DummySendingClientId};
-		
-		TSharedRef<ConcertSyncCore::FObjectReplicationCache> Cache = MakeShared<ConcertSyncCore::FObjectReplicationCache>(MakeShared<FTestReplicationFormat>());
+
+		FTestReplicationFormat TestFormat;
+		TSharedRef<ConcertSyncCore::FObjectReplicationCache> Cache = MakeShared<ConcertSyncCore::FObjectReplicationCache>(TestFormat);
 		TSharedRef<FTestReplicationCacheUser> User_NeverConsume = MakeShared<FTestReplicationCacheUser>(*this, ObjectID, EReplicationCacheTestFlags::NeverConsume);
 		TSharedRef<FTestReplicationCacheUser> User_NeverReceive = MakeShared<FTestReplicationCacheUser>(*this, ObjectID, EReplicationCacheTestFlags::NeverReceive);
 		TSharedRef<FTestReplicationCacheUser> User_ConsumeManually = MakeShared<FTestReplicationCacheUser>(*this, ObjectID);
diff --git a/Engine/Plugins/Developer/Concert/ConcertSync/ConcertSyncTest/Source/ConcertSyncTest/Private/Replication/ReplicationHandshakeTests.cpp b/Engine/Plugins/Developer/Concert/ConcertSync/ConcertSyncTest/Source/ConcertSyncTest/Private/Replication/ReplicationHandshakeTests.cpp
--- a/Engine/Plugins/Developer/Concert/ConcertSync/ConcertSyncTest/Source/ConcertSyncTest/Private/Replication/ReplicationHandshakeTests.cpp
+++ b/Engine/Plugins/Developer/Concert/ConcertSync/ConcertSyncTest/Source/ConcertSyncTest/Private/Replication/ReplicationHandshakeTests.cpp
@@ -42,11 +42,11 @@ namespace UE::ConcertSyncTests::Replication::Handshake
 		FClientInfo& Client = ConnectClient();
 		const TSharedRef<IConcertClientReplicationBridge> BridgeMock = MakeShared<FConcertClientReplicationBridgeMock>();
 		const TSharedPtr<IConcertClientSession>& ClientSession = Client.ClientSessionMock;
-		const TSharedRef<IConcertClientReplicationManager> ClientReplicationManager_Primary = CreateClientReplicationManager(ClientSession.ToSharedRef(), &BridgeMock.Get());
+		const TSharedRef<IConcertClientReplicationManager> ClientReplicationManager_Primary = CreateClientReplicationManager(ClientSession.ToSharedRef(), BridgeMock.Get());
 		FClientInfo& Client_Secondary = ConnectClient();
 		const TSharedRef<IConcertClientReplicationBridge> BridgeMock_Secondary = MakeShared<FConcertClientReplicationBridgeMock>();
 		const TSharedPtr<IConcertClientSession>& ClientSession_Secondary = Client_Secondary.ClientSessionMock;
-		const TSharedRef<IConcertClientReplicationManager> ClientReplicationManager_Secondary = CreateClientReplicationManager(ClientSession_Secondary.ToSharedRef(), &BridgeMock_Secondary.Get());
+		const TSharedRef<IConcertClientReplicationManager> ClientReplicationManager_Secondary = CreateClientReplicationManager(ClientSession_Secondary.ToSharedRef(), BridgeMock_Secondary.Get());
 
 		// Stream 
 		FConcertReplicationStream StreamDescription;
diff --git a/Engine/Plugins/Developer/Concert/ConcertSync/ConcertSyncTest/Source/ConcertSyncTest/Private/Replication/ReplicationTestInterface.h b/Engine/Plugins/Developer/Concert/ConcertSync/ConcertSyncTest/Source/ConcertSyncTest/Private/Replication/ReplicationTestInterface.h
--- a/Engine/Plugins/Developer/Concert/ConcertSync/ConcertSyncTest/Source/ConcertSyncTest/Private/Replication/ReplicationTestInterface.h
+++ b/Engine/Plugins/Developer/Concert/ConcertSync/ConcertSyncTest/Source/ConcertSyncTest/Private/Replication/ReplicationTestInterface.h
@@ -18,7 +18,7 @@ namespace UE::ConcertSyncClient::TestInterface
 {
 	extern CONCERTSYNCCLIENT_API TSharedRef<IConcertClientReplicationManager> CreateClientReplicationManager(
 		TSharedRef<IConcertClientSession> InLiveSession,
-		IConcertClientReplicationBridge* InBridge
+		IConcertClientReplicationBridge& InBridge UE_LIFETIMEBOUND
 		);
 
 	extern CONCERTSYNCCLIENT_API TSharedRef<IConcertClientReplicationBridge> CreateClientReplicationBridge();
diff --git a/Engine/Plugins/Developer/Concert/ConcertSync/ConcertSyncTest/Source/ConcertSyncTest/Private/Replication/Util/SendReceiveTestBase.cpp b/Engine/Plugins/Developer/Concert/ConcertSync/ConcertSyncTest/Source/ConcertSyncTest/Private/Replication/Util/SendReceiveTestBase.cpp
--- a/Engine/Plugins/Developer/Concert/ConcertSync/ConcertSyncTest/Source/ConcertSyncTest/Private/Replication/Util/SendReceiveTestBase.cpp
+++ b/Engine/Plugins/Developer/Concert/ConcertSync/ConcertSyncTest/Source/ConcertSyncTest/Private/Replication/Util/SendReceiveTestBase.cpp
@@ -63,7 +63,7 @@ namespace UE::ConcertSyncTests::Replication
 			BridgeMock_Receiver =  MakeShared<FConcertClientReplicationBridgeMock>();
 			BridgeUsed_Receiver = BridgeMock_Receiver;
 		}
-		ClientReplicationManager_Receiver = ConcertSyncClient::TestInterface::CreateClientReplicationManager(Client_Receiver->ClientSessionMock, BridgeUsed_Receiver.Get());
+		ClientReplicationManager_Receiver = ConcertSyncClient::TestInterface::CreateClientReplicationManager(Client_Receiver->ClientSessionMock, *BridgeUsed_Receiver.Get());
 		
 		// Client Sender
 		Client_Sender = &ConnectClient();
@@ -76,7 +76,7 @@ namespace UE::ConcertSyncTests::Replication
 			BridgeMock_Sender =  MakeShared<FConcertClientReplicationBridgeMock>();
 			BridgeUsed_Sender = BridgeMock_Sender;
 		}
-		ClientReplicationManager_Sender = ConcertSyncClient::TestInterface::CreateClientReplicationManager(Client_Sender->ClientSessionMock, BridgeUsed_Sender.Get());
+		ClientReplicationManager_Sender = ConcertSyncClient::TestInterface::CreateClientReplicationManager(Client_Sender->ClientSessionMock, *BridgeUsed_Sender.Get());
 		
 		// 1.1 Sender offers to send all UTestReflectionObject properties
 		{
