Fixed Concert updates no longer marking the updated objects as dirty, preventing them from being saved and leading to data loss

This was fallout from 32945747, which set GIsTransacting around the PreEditChange calls from within Concert. This resulted in UObject::PreEditChange calling Modify with bShouldMarkAsDirty set to false.

To resolve this we now call Modify ourselves prior to modifying the object state from within Concert.

#rnx

diff --git a/Engine/Plugins/Developer/Concert/ConcertSync/ConcertSyncClient/Source/ConcertSyncClient/Private/ConcertClientTransactionBridge.cpp b/Engine/Plugins/Developer/Concert/ConcertSync/ConcertSyncClient/Source/ConcertSyncClient/Private/ConcertClientTransactionBridge.cpp
--- a/Engine/Plugins/Developer/Concert/ConcertSync/ConcertSyncClient/Source/ConcertSyncClient/Private/ConcertClientTransactionBridge.cpp
+++ b/Engine/Plugins/Developer/Concert/ConcertSync/ConcertSyncClient/Source/ConcertSyncClient/Private/ConcertClientTransactionBridge.cpp
@@ -544,6 +544,8 @@ void ProcessTransactionEvent(const FConcertTransactionEventBase& InEvent, const
 			continue;
 		}
 
+		TransactionObject->Modify();
+
 		// Apply the new data
 		if (ObjectUpdate.ObjectData.SerializedData.Num() > 0)
 		{
