Add mechanism that allows external code to tell Multi-user to bypass the hot reload for certain packages.  This will speed up sessions by skipping unnecessary packages for hot reload.

diff --git a/Engine/Plugins/Developer/Concert/ConcertSync/ConcertSyncClient/Source/ConcertSyncClient/Private/ConcertClientPackageBridge.cpp b/Engine/Plugins/Developer/Concert/ConcertSync/ConcertSyncClient/Source/ConcertSyncClient/Private/ConcertClientPackageBridge.cpp
--- a/Engine/Plugins/Developer/Concert/ConcertSync/ConcertSyncClient/Source/ConcertSyncClient/Private/ConcertClientPackageBridge.cpp
+++ b/Engine/Plugins/Developer/Concert/ConcertSync/ConcertSyncClient/Source/ConcertSyncClient/Private/ConcertClientPackageBridge.cpp
@@ -6,6 +6,7 @@
 #include "ConcertWorkspaceData.h"
 #include "ConcertSyncClientUtil.h"
 
+#include "Algo/AllOf.h"
 #include "Engine/World.h"
 #include "Engine/Engine.h"
 #include "IConcertClientPackageBridge.h"
@@ -165,8 +166,7 @@ void FConcertClientPackageBridge::HandlePackagePreSave(UPackage* Package, FObjec
 	{
 		if (IFileManager::Get().FileExists(*PackageFilename))
 		{
-			FConcertPackageInfo PackageInfo;
-			ConcertSyncClientUtil::FillPackageInfo(Package, Asset, EConcertPackageUpdateType::Saved, PackageInfo);
+			FConcertPackageInfo PackageInfo = ConcertSyncClientUtil::FillPackageInfo(Package, Asset, EConcertPackageUpdateType::Saved);
 			PackageInfo.bPreSave = true;
 			PackageInfo.bAutoSave = GEngine->IsAutosaving();
 
@@ -220,12 +220,11 @@ void FConcertClientPackageBridge::HandlePackageSaved(const FString& PackageFilen
 	FName NewPackageName;
 	PackagesBeingRenamed.RemoveAndCopyValue(Package->GetFName(), NewPackageName);
 
-	FConcertPackageInfo PackageInfo;
-	ConcertSyncClientUtil::FillPackageInfo(Package, nullptr, NewPackageName.IsNone() ? EConcertPackageUpdateType::Saved : EConcertPackageUpdateType::Renamed, PackageInfo);
+	FConcertPackageInfo PackageInfo = ConcertSyncClientUtil::FillPackageInfo(Package, nullptr, NewPackageName.IsNone() ? EConcertPackageUpdateType::Saved : EConcertPackageUpdateType::Renamed);
 	PackageInfo.NewPackageName = NewPackageName;
 	PackageInfo.bPreSave = false;
 	PackageInfo.bAutoSave = GEngine->IsAutosaving();
-
+	PackageInfo.bCanSkipHotReload = CanSkipHotReload(PackageInfo);
 	PendingPackageInfos.Add(FConcertPackageInfoTuple(MoveTemp(PackageInfo), PackageFilename));
 
 	UE_LOG(LogConcert, Verbose, TEXT("Asset Saved: %s"), *Package->GetName());
@@ -271,9 +270,8 @@ void FConcertClientPackageBridge::HandleAssetAdded(UObject *Object)
 
 			if (IFileManager::Get().FileExists(*PackageFilename))
 			{
-				FConcertPackageInfo PackageInfo;
-				ConcertSyncClientUtil::FillPackageInfo(Package, Asset, EConcertPackageUpdateType::Added, PackageInfo);
-
+				FConcertPackageInfo PackageInfo = ConcertSyncClientUtil::FillPackageInfo(Package, Asset, EConcertPackageUpdateType::Added);
+				PackageInfo.bCanSkipHotReload = CanSkipHotReload(PackageInfo);
 				OnLocalPackageEventDelegate.Broadcast(PackageInfo, PackageFilename);
 				IFileManager::Get().Delete(*PackageFilename);
 			}
@@ -303,8 +301,7 @@ void FConcertClientPackageBridge::HandleAssetDeleted(UObject *Object)
 
 	UPackage* Package = Object->GetOutermost();
 
-	FConcertPackageInfo PackageInfo;
-	ConcertSyncClientUtil::FillPackageInfo(Package, nullptr, EConcertPackageUpdateType::Deleted, PackageInfo);
+	const FConcertPackageInfo PackageInfo = ConcertSyncClientUtil::FillPackageInfo(Package, nullptr, EConcertPackageUpdateType::Deleted);
 	OnLocalPackageEventDelegate.Broadcast(PackageInfo, FString());
 
 	UE_LOG(LogConcert, Verbose, TEXT("Asset Deleted: %s"), *Package->GetName());
@@ -364,6 +361,26 @@ void FConcertClientPackageBridge::HandleMapChanged(UWorld* InWorld, EMapChangeTy
 }
 #endif
 
+void FConcertClientPackageBridge::RegisterPackageHotReloadHint(FName PackageReloadHintName, FPackageHotReloadHintDelegate FilterHandle)
+{
+	check(HotReloadHints.Find(PackageReloadHintName) == nullptr);
+	check(FilterHandle.IsBound());
+	HotReloadHints.Add(PackageReloadHintName) = MoveTemp(FilterHandle);
+}
+
+bool FConcertClientPackageBridge::CanSkipHotReload(const FConcertPackageInfo& PackageInfo) const
+{
+	return Algo::AllOf(HotReloadHints, [&PackageInfo](const TPair<FName,FPackageHotReloadHintDelegate>& Item)
+	{
+		return Item.Value.Execute(PackageInfo);
+	});
+}
+
+void FConcertClientPackageBridge::UnregisterPackageHotReloadHint(FName PackageReloadHintName)
+{
+	HotReloadHints.FindAndRemoveChecked(PackageReloadHintName);
+}
+
 void FConcertClientPackageBridge::RegisterPackageFilter(FName FilterName, FPackageFilterDelegate FilterHandle)
 {
 	check(PackageFilters.Find(FilterName) == nullptr);
diff --git a/Engine/Plugins/Developer/Concert/ConcertSync/ConcertSyncClient/Source/ConcertSyncClient/Private/ConcertClientPackageBridge.h b/Engine/Plugins/Developer/Concert/ConcertSync/ConcertSyncClient/Source/ConcertSyncClient/Private/ConcertClientPackageBridge.h
--- a/Engine/Plugins/Developer/Concert/ConcertSync/ConcertSyncClient/Source/ConcertSyncClient/Private/ConcertClientPackageBridge.h
+++ b/Engine/Plugins/Developer/Concert/ConcertSync/ConcertSyncClient/Source/ConcertSyncClient/Private/ConcertClientPackageBridge.h
@@ -27,6 +27,10 @@ public:
 	virtual bool& GetIgnoreLocalSaveRef() override;
 	virtual bool& GetIgnoreLocalDiscardRef() override;
 
+	virtual void RegisterPackageHotReloadHint(FName PackageReloadHintName, FPackageHotReloadHintDelegate FilterHandle) override;
+	virtual bool CanSkipHotReload(const FConcertPackageInfo& PackageInfo) const override;
+	virtual void UnregisterPackageHotReloadHint(FName PackageReloadHintName) override;
+
 	virtual void RegisterPackageFilter(FName FilterName, FPackageFilterDelegate FilterHandle) override;
 	virtual void UnregisterPackageFilter(FName FilterName) override;
 	virtual EPackageFilterResult IsPackageFiltered(const FConcertPackageInfo& PackageInfo) const override;
@@ -68,6 +72,9 @@ private:
 	/** Map of named packages filters that can override what is included / excluded by package bridge*/
 	TMap<FName, FPackageFilterDelegate> PackageFilters;
 
+	/** Map of named hot reload delegates that can provide hints on if packages can skip hot reload, e.g. take records. */
+	TMap<FName, FPackageHotReloadHintDelegate> HotReloadHints;
+
 	/** Flag to ignore package change events, used when we do not want to record package changes we generate ourselves */
 	bool bIgnoreLocalSave;
 
diff --git a/Engine/Plugins/Developer/Concert/ConcertSync/ConcertSyncClient/Source/ConcertSyncClient/Private/ConcertClientPackageManager.cpp b/Engine/Plugins/Developer/Concert/ConcertSync/ConcertSyncClient/Source/ConcertSyncClient/Private/ConcertClientPackageManager.cpp
--- a/Engine/Plugins/Developer/Concert/ConcertSync/ConcertSyncClient/Source/ConcertSyncClient/Private/ConcertClientPackageManager.cpp
+++ b/Engine/Plugins/Developer/Concert/ConcertSync/ConcertSyncClient/Source/ConcertSyncClient/Private/ConcertClientPackageManager.cpp
@@ -306,10 +306,8 @@ void FConcertClientPackageManager::ApplyAllHeadPackageData()
 
 bool FConcertClientPackageManager::PassesPackageFilters(UPackage* InPackage) const
 {
-	// Create a dummy package info to run filters on
-	FConcertPackageInfo PackageInfo;
-	ConcertSyncClientUtil::FillPackageInfo(InPackage, nullptr, EConcertPackageUpdateType::Saved, PackageInfo);
-	return ApplyPackageFilters(PackageInfo);
+	FConcertPackageInfo DummyPackageInfo = ConcertSyncClientUtil::FillPackageInfo(InPackage, nullptr, EConcertPackageUpdateType::Saved);
+	return ApplyPackageFilters(DummyPackageInfo);
 }
 
 bool FConcertClientPackageManager::HasSessionChanges() const
@@ -476,6 +474,7 @@ void FConcertClientPackageManager::AddPendingReloadForNewExternalMaps(const FCon
 			&& !PackagesPendingHotReload.Contains(PackageName))
 		{
 			UE_LOG(LogConcert, Display, TEXT("Scheduling reloading for world partition persistent level %s."), *PackagePathname);
+
 			PackagesPendingHotReload.Add(PackageName);
 		}
 	}
@@ -610,7 +609,11 @@ void FConcertClientPackageManager::SavePackageFile(const FConcertPackageInfo& Pa
 
 	if (bSuccess)
 	{
-		PackagesPendingHotReload.Add(PackageInfo.PackageName);
+		if (!PackageInfo.bCanSkipHotReload)
+		{
+			// We skip hot reloading assets that are hinted not to need it. @see FConcertPackageInfo::bCanSkipHotReload)
+			PackagesPendingHotReload.Add(PackageInfo.PackageName);
+		}
 		PackagesPendingPurge.Remove(PackageInfo.PackageName);
 	}
 }
diff --git a/Engine/Plugins/Developer/Concert/ConcertSync/ConcertSyncClient/Source/ConcertSyncClient/Private/ConcertSyncClientUtil.cpp b/Engine/Plugins/Developer/Concert/ConcertSync/ConcertSyncClient/Source/ConcertSyncClient/Private/ConcertSyncClientUtil.cpp
--- a/Engine/Plugins/Developer/Concert/ConcertSync/ConcertSyncClient/Source/ConcertSyncClient/Private/ConcertSyncClientUtil.cpp
+++ b/Engine/Plugins/Developer/Concert/ConcertSync/ConcertSyncClient/Source/ConcertSyncClient/Private/ConcertSyncClientUtil.cpp
@@ -783,13 +783,15 @@ bool IsWorldPartitionWorld()
 	return false;
 }
 
-void FillPackageInfo(UPackage* InPackage, UObject* InAsset, const EConcertPackageUpdateType InPackageUpdateType, FConcertPackageInfo& OutPackageInfo)
+FConcertPackageInfo FillPackageInfo(UPackage* InPackage, UObject* InAsset, const EConcertPackageUpdateType InPackageUpdateType)
 {
+	FConcertPackageInfo OutPackageInfo;
 	UObject* Asset = InAsset ? InAsset : InPackage->FindAssetInPackage();
 	OutPackageInfo.PackageName = InPackage->GetFName();
 	OutPackageInfo.AssetClass = Asset ? Asset->GetClass()->GetPathName() : FString();
 	OutPackageInfo.PackageFileExtension = Asset && Asset->IsA<UWorld>()? FPackageName::GetMapPackageExtension() : FPackageName::GetAssetPackageExtension();
 	OutPackageInfo.PackageUpdateType = InPackageUpdateType;
+	return OutPackageInfo;
 }
 
 }
diff --git a/Engine/Plugins/Developer/Concert/ConcertSync/ConcertSyncClient/Source/ConcertSyncClient/Private/ConcertSyncClientUtil.h b/Engine/Plugins/Developer/Concert/ConcertSync/ConcertSyncClient/Source/ConcertSyncClient/Private/ConcertSyncClientUtil.h
--- a/Engine/Plugins/Developer/Concert/ConcertSync/ConcertSyncClient/Source/ConcertSyncClient/Private/ConcertSyncClientUtil.h
+++ b/Engine/Plugins/Developer/Concert/ConcertSync/ConcertSyncClient/Source/ConcertSyncClient/Private/ConcertSyncClientUtil.h
@@ -151,5 +151,5 @@ namespace ConcertSyncClientUtil
 	/** Return true if we are in a world partition world. */
 	bool IsWorldPartitionWorld();
 
-	void FillPackageInfo(UPackage* InPackage, UObject* InAsset, const EConcertPackageUpdateType InPackageUpdateType, FConcertPackageInfo& OutPackageInfo);
+	FConcertPackageInfo FillPackageInfo(UPackage* InPackage, UObject* InAsset, const EConcertPackageUpdateType InPackageUpdateType);
 }
diff --git a/Engine/Plugins/Developer/Concert/ConcertSync/ConcertSyncClient/Source/ConcertSyncClient/Public/IConcertClientPackageBridge.h b/Engine/Plugins/Developer/Concert/ConcertSync/ConcertSyncClient/Source/ConcertSyncClient/Public/IConcertClientPackageBridge.h
--- a/Engine/Plugins/Developer/Concert/ConcertSync/ConcertSyncClient/Source/ConcertSyncClient/Public/IConcertClientPackageBridge.h
+++ b/Engine/Plugins/Developer/Concert/ConcertSync/ConcertSyncClient/Source/ConcertSyncClient/Public/IConcertClientPackageBridge.h
@@ -19,6 +19,7 @@ enum class EPackageFilterResult : uint8
 };
 
 DECLARE_DELEGATE_RetVal_OneParam(EPackageFilterResult, FPackageFilterDelegate, const FConcertPackageInfo&);
+DECLARE_DELEGATE_RetVal_OneParam(bool, FPackageHotReloadHintDelegate, const FConcertPackageInfo&);
 
 /**
  * Bridge between the editor package events and Concert.
@@ -57,6 +58,15 @@ public:
 	 */
 	virtual FOnConcertClientLocalPackageDiscarded& OnLocalPackageDiscarded() = 0;
 
+	/** Register a pending package hot reload hint package payloads. */
+	virtual void RegisterPackageHotReloadHint(FName PackageReloadHintName, FPackageHotReloadHintDelegate FilterHandle) = 0;
+
+	/** Returns true if the package can be skipped for hot reload on client nodes.  */
+	virtual bool CanSkipHotReload(const FConcertPackageInfo& PackageInfo) const = 0;
+
+	/** Unregister package hot reload hint handler. */
+	virtual void UnregisterPackageHotReloadHint(FName PackageReloadHintName) = 0;
+
 	/** Register a package filter to exclude/include certain packagins from the session. */
 	virtual void RegisterPackageFilter(FName FilterName, FPackageFilterDelegate FilterHandle) = 0;
 
diff --git a/Engine/Plugins/Developer/Concert/ConcertSync/ConcertSyncCore/Source/ConcertSyncCore/Public/ConcertWorkspaceData.h b/Engine/Plugins/Developer/Concert/ConcertSync/ConcertSyncCore/Source/ConcertSyncCore/Public/ConcertWorkspaceData.h
--- a/Engine/Plugins/Developer/Concert/ConcertSync/ConcertSyncCore/Source/ConcertSyncCore/Public/ConcertWorkspaceData.h
+++ b/Engine/Plugins/Developer/Concert/ConcertSync/ConcertSyncCore/Source/ConcertSyncCore/Public/ConcertWorkspaceData.h
@@ -58,6 +58,16 @@ struct FConcertPackageInfo
 	/** Was this update caused by an auto-save? */
 	UPROPERTY()
 	bool bAutoSave = false;
+
+	/**
+	 * Can we skip hot reload on this package? Hot reload is required to update changes to a loaded asset to a newer versoin
+	 * modified in a multi-user session.  Take recordings are packages that do not need to go through hot reload as they have no live assets.
+	 * This flag can be set to true to indicate to the replay system the package can bypass hot reload.
+	 *
+	 * @see IConcertClientClientPackageBridge::RegisterPackageHotReloadHint
+	 */
+	UPROPERTY()
+	bool bCanSkipHotReload = false;
 };
 
 USTRUCT()
diff --git a/Engine/Plugins/VirtualProduction/MultiUserTakes/Source/ConcertTakeRecorder/Private/ConcertTakeRecorderManager.cpp b/Engine/Plugins/VirtualProduction/MultiUserTakes/Source/ConcertTakeRecorder/Private/ConcertTakeRecorderManager.cpp
--- a/Engine/Plugins/VirtualProduction/MultiUserTakes/Source/ConcertTakeRecorder/Private/ConcertTakeRecorderManager.cpp
+++ b/Engine/Plugins/VirtualProduction/MultiUserTakes/Source/ConcertTakeRecorder/Private/ConcertTakeRecorderManager.cpp
@@ -248,6 +248,7 @@ void FConcertTakeRecorderManager::RegisterExtensions()
 			IConcertClientPackageBridge* PackageBridge = ConcertSyncClient->GetPackageBridge();
 			check(PackageBridge);
 			PackageBridge->RegisterPackageFilter(TEXT("ConcertTakes"), FPackageFilterDelegate::CreateRaw(this, &FConcertTakeRecorderManager::ShouldPackageBeFiltered));
+			PackageBridge->RegisterPackageHotReloadHint(TEXT("ConcertTakes"), FPackageHotReloadHintDelegate::CreateRaw(this, &FConcertTakeRecorderManager::CanSkipHotReload));
 		}
 
 		FPropertyEditorModule& PropertyEditorModule = FModuleManager::Get().LoadModuleChecked<FPropertyEditorModule>("PropertyEditor");
@@ -304,6 +305,7 @@ void FConcertTakeRecorderManager::UnregisterExtensions()
 			IConcertClientPackageBridge* PackageBridge = ConcertSyncClient->GetPackageBridge();
 			check(PackageBridge != nullptr);
 			PackageBridge->UnregisterPackageFilter(TEXT("ConcertTakes"));
+			PackageBridge->UnregisterPackageHotReloadHint(TEXT("ConcertTakes"));
 		}
 	}
 }
@@ -969,6 +971,17 @@ EPackageFilterResult FConcertTakeRecorderManager::ShouldPackageBeFiltered(const
 	return EPackageFilterResult::UseDefault;
 }
 
+bool FConcertTakeRecorderManager::CanSkipHotReload(const FConcertPackageInfo& InPackageInfo)
+{
+	if (IsTakeSyncEnabled() && ensure(WeakSession.IsValid()) && CanRecord())
+	{
+		FTakeRecorderProjectParameters Project = GetDefault<UTakeRecorderProjectSettings>()->Settings;
+		FString FullName = InPackageInfo.PackageName.ToString();
+		return FullName.Contains(Project.RootTakeSaveDir.Path);
+	}
+	return false;
+}
+
 ETransactionFilterResult FConcertTakeRecorderManager::ShouldObjectBeTransacted(UObject* InObject, UPackage* InPackage)
 {
 	UConcertSessionRecordSettings const* RecordSettings = GetDefault<UConcertSessionRecordSettings>();
diff --git a/Engine/Plugins/VirtualProduction/MultiUserTakes/Source/ConcertTakeRecorder/Private/ConcertTakeRecorderManager.h b/Engine/Plugins/VirtualProduction/MultiUserTakes/Source/ConcertTakeRecorder/Private/ConcertTakeRecorderManager.h
--- a/Engine/Plugins/VirtualProduction/MultiUserTakes/Source/ConcertTakeRecorder/Private/ConcertTakeRecorderManager.h
+++ b/Engine/Plugins/VirtualProduction/MultiUserTakes/Source/ConcertTakeRecorder/Private/ConcertTakeRecorderManager.h
@@ -120,7 +120,7 @@ private:
 
 	ETransactionFilterResult ShouldObjectBeTransacted(UObject* InObject, UPackage* InPackage);
 	EPackageFilterResult ShouldPackageBeFiltered(const FConcertPackageInfo& InPackage);
-
+	bool CanSkipHotReload(const FConcertPackageInfo& InPackage);
 private:
 	FTakeRecorderParameters SetupTakeParametersForMultiuser(const FTakeRecorderParameters& Input);
 
