Deprecate old version of RegisterTransactionFilter

diff --git a/Engine/Plugins/Developer/Concert/ConcertSync/ConcertSyncClient/Source/ConcertSyncClient/Public/IConcertClientTransactionBridge.h b/Engine/Plugins/Developer/Concert/ConcertSync/ConcertSyncClient/Source/ConcertSyncClient/Public/IConcertClientTransactionBridge.h
--- a/Engine/Plugins/Developer/Concert/ConcertSync/ConcertSyncClient/Source/ConcertSyncClient/Public/IConcertClientTransactionBridge.h
+++ b/Engine/Plugins/Developer/Concert/ConcertSync/ConcertSyncClient/Source/ConcertSyncClient/Public/IConcertClientTransactionBridge.h
@@ -90,6 +90,7 @@ struct FConcertTransactionFilterArgs
 };
 
 DECLARE_MULTICAST_DELEGATE_TwoParams(FOnApplyTransaction, ETransactionNotification, const bool bIsSnapshot);
+UE_DEPRECATED(5.5, "Use FOnFilterTransactionDelegate instead.")
 DECLARE_DELEGATE_RetVal_TwoParams(ETransactionFilterResult, FTransactionFilterDelegate, UObject*, UPackage*);
 DECLARE_DELEGATE_RetVal_OneParam(ETransactionFilterResult, FOnFilterTransactionDelegate, const FConcertTransactionFilterArgs&);
 DECLARE_MULTICAST_DELEGATE_TwoParams(FOnConcertClientLocalTransactionSnapshot, const FConcertClientLocalTransactionCommonData&, const FConcertClientLocalTransactionSnapshotData&);
@@ -188,6 +189,7 @@ public:
 	virtual void ApplyRemoteTransaction(const FConcertTransactionEventBase& InEvent, const FConcertSessionVersionInfo* InVersionInfo, const TArray<FName>& InPackagesToProcess, const FConcertLocalIdentifierTable* InLocalIdentifierTablePtr, const bool bIsSnapshot, const class FConcertSyncWorldRemapper& ConcertSyncWorldRemapper) = 0;
 
 	/** Callback to register delegate for handling transaction events */
+	UE_DEPRECATED(5.5, "Use the version of RegisterTransactionFilter accepting FOnFilterTransactionDelegate instead.")
 	void RegisterTransactionFilter(
 		FName FilterName,
 		PRAGMA_DISABLE_DEPRECATION_WARNINGS
