Always return the file path expected as if the data is in the content area.

#rnx

diff --git a/Engine/Plugins/Developer/Concert/ConcertSync/ConcertSyncClient/Source/ConcertSyncClient/Private/ConcertSandboxPlatformFile.cpp b/Engine/Plugins/Developer/Concert/ConcertSync/ConcertSyncClient/Source/ConcertSyncClient/Private/ConcertSandboxPlatformFile.cpp
--- a/Engine/Plugins/Developer/Concert/ConcertSync/ConcertSyncClient/Source/ConcertSyncClient/Private/ConcertSandboxPlatformFile.cpp
+++ b/Engine/Plugins/Developer/Concert/ConcertSync/ConcertSyncClient/Source/ConcertSyncClient/Private/ConcertSandboxPlatformFile.cpp
@@ -410,15 +410,6 @@ FDateTime FConcertSandboxPlatformFile::GetAccessTimeStamp(const TCHAR* Filename)
 FString FConcertSandboxPlatformFile::GetFilenameOnDisk(const TCHAR* Filename)
 {
 	const FConcertSandboxPlatformFilePath ResolvedPath = ToSandboxPath(Filename);
-
-	if (ResolvedPath.HasSandboxPath())
-	{
-		if (IsPathDeleted(ResolvedPath) || LowerLevel->FileExists(*ResolvedPath.GetSandboxPath()))
-		{
-			return LowerLevel->GetFilenameOnDisk(*ResolvedPath.GetSandboxPath());
-		}
-	}
-
 	return LowerLevel->GetFilenameOnDisk(*ResolvedPath.GetNonSandboxPath());
 }
 
