Prevent checkout dialog for duplicate actors when in a Multi-user session.

diff --git a/Engine/Plugins/Developer/Concert/ConcertSync/ConcertSyncClient/Source/ConcertSyncClient/Private/ConcertClientPackageManager.cpp b/Engine/Plugins/Developer/Concert/ConcertSync/ConcertSyncClient/Source/ConcertSyncClient/Private/ConcertClientPackageManager.cpp
--- a/Engine/Plugins/Developer/Concert/ConcertSync/ConcertSyncClient/Source/ConcertSyncClient/Private/ConcertClientPackageManager.cpp
+++ b/Engine/Plugins/Developer/Concert/ConcertSync/ConcertSyncClient/Source/ConcertSyncClient/Private/ConcertClientPackageManager.cpp
@@ -89,6 +89,9 @@ FConcertClientPackageManager::FConcertClientPackageManager(TSharedRef<FConcertSy
 	//
 	UPackage::PackageDirtyStateChangedEvent.AddRaw(this, &FConcertClientPackageManager::HandlePackageDirtyStateChanged);
 	PackageBridge->OnLocalPackageEvent().AddRaw(this, &FConcertClientPackageManager::HandleLocalPackageEvent);
+
+	FEditorDelegates::OnDuplicateActorsBegin.AddRaw(this, &FConcertClientPackageManager::HandleDuplicateActorsBegin);
+	FEditorDelegates::OnDuplicateActorsEnd.AddRaw(this, &FConcertClientPackageManager::HandleDuplicateActorsEnd);
 #endif	// WITH_EDITOR
 
 	LiveSession->GetSession().RegisterCustomEventHandler<FConcertPackageRejectedEvent>(this, &FConcertClientPackageManager::HandlePackageRejectedEvent);
@@ -107,6 +110,8 @@ FConcertClientPackageManager::~FConcertClientPackageManager()
 		SandboxPlatformFile->DiscardSandbox(PackagesPendingHotReload, PackagesPendingPurge);
 		SandboxPlatformFile.Reset();
 	}
+	FEditorDelegates::OnDuplicateActorsBegin.RemoveAll(this);
+	FEditorDelegates::OnDuplicateActorsEnd.RemoveAll(this);
 #endif	// WITH_EDITOR
 
 	LiveSession->GetSession().UnregisterCustomEventHandler<FConcertPackageRejectedEvent>(this);
@@ -266,6 +271,27 @@ void FConcertClientPackageManager::SynchronizeInMemoryPackages()
 	HotReloadPendingPackages();
 }
 
+#if WITH_EDITOR
+void FConcertClientPackageManager::HandleDuplicateActorsBegin()
+{
+	UWorld* CurrentWorld = ConcertSyncClientUtil::GetCurrentWorld();
+	check(CurrentWorld);
+
+	ULevel* Level = CurrentWorld->GetCurrentLevel();
+	bPromptWhenAddingToLevelBeforeCheckout = Level->bPromptWhenAddingToLevelBeforeCheckout;
+	Level->bPromptWhenAddingToLevelBeforeCheckout = false;
+}
+
+void FConcertClientPackageManager::HandleDuplicateActorsEnd()
+{
+	UWorld* CurrentWorld = ConcertSyncClientUtil::GetCurrentWorld();
+	check(CurrentWorld);
+
+	ULevel* Level = CurrentWorld->GetCurrentLevel();
+	Level->bPromptWhenAddingToLevelBeforeCheckout = bPromptWhenAddingToLevelBeforeCheckout;
+}
+#endif
+
 void FConcertClientPackageManager::HandlePackageDiscarded(UPackage* InPackage)
 {
 	FConcertPackageUpdateEvent Event;
diff --git a/Engine/Plugins/Developer/Concert/ConcertSync/ConcertSyncClient/Source/ConcertSyncClient/Private/ConcertClientPackageManager.h b/Engine/Plugins/Developer/Concert/ConcertSync/ConcertSyncClient/Source/ConcertSyncClient/Private/ConcertClientPackageManager.h
--- a/Engine/Plugins/Developer/Concert/ConcertSync/ConcertSyncClient/Source/ConcertSyncClient/Private/ConcertClientPackageManager.h
+++ b/Engine/Plugins/Developer/Concert/ConcertSync/ConcertSyncClient/Source/ConcertSyncClient/Private/ConcertClientPackageManager.h
@@ -146,6 +146,18 @@ private:
 	 */
 	void HandleLocalPackageEvent(const FConcertPackageInfo& PackageInfo, const FString& PackagePathname);
 
+#if WITH_EDITOR
+	/**
+	 * Handle begin event when a duplicate actor is added to the level.
+	 */
+	void HandleDuplicateActorsBegin();
+
+	/**
+	 * Handle end event when a duplicate actor is added to the level.
+	 */
+	void HandleDuplicateActorsEnd();
+#endif
+
 	/**
 	 * Utility to save new package data to disk, and also queue if for hot-reload.
 	 */
@@ -199,6 +211,12 @@ private:
 	 */
 	IConcertClientPackageBridge* PackageBridge;
 
+	/**
+	 * Value of the CurrentWorld()->GetCurrentLevel()->bPromptWhenAddingToLevelBeforeCheckout on the level. We suppress this
+	 * for sandbox sessions.  This value is only relevant between begin / end delegate on duplicate actors event.
+	 */
+	bool bPromptWhenAddingToLevelBeforeCheckout = false;
+
 	/**
 	 * Indicates if we are currently hot reloading.
 	 */
