Fix client not removing authority from cache when request removes stream.
Added a unit test for this case.

diff --git a/Engine/Plugins/Developer/Concert/ConcertSync/ConcertSyncClient/Source/ConcertSyncClient/Private/Replication/Manager/ReplicationManagerState_Connected.cpp b/Engine/Plugins/Developer/Concert/ConcertSync/ConcertSyncClient/Source/ConcertSyncClient/Private/Replication/Manager/ReplicationManagerState_Connected.cpp
--- a/Engine/Plugins/Developer/Concert/ConcertSync/ConcertSyncClient/Source/ConcertSyncClient/Private/Replication/Manager/ReplicationManagerState_Connected.cpp
+++ b/Engine/Plugins/Developer/Concert/ConcertSync/ConcertSyncClient/Source/ConcertSyncClient/Private/Replication/Manager/ReplicationManagerState_Connected.cpp
@@ -223,13 +223,13 @@ namespace UE::ConcertSyncClient::Replication
 		
 		// Stop replicating removed objects right now: the server will remove authority after processing this request.
 		// At that point, it will log errors for receiving replication data from a client without authority.
-		HandleRemovingReplicatedObjects(Args);
+		TMap<FSoftObjectPath, TArray<FGuid>> PredictedChange = HandleRemovingReplicatedObjects(Args);
 		// We don't need worry about updating sync control until it is processed below - the local client will not attempt to replicate the object
 		// because we just locally updated the replication cache.
 		
 		Private::LogNetworkMessage(CVarLogStreamRequestsAndResponsesOnClient, Args);
 		return LiveSession->SendCustomRequest<FConcertReplication_ChangeStream_Request, FConcertReplication_ChangeStream_Response>(Args, LiveSession->GetSessionServerEndpointId())
-			.Next([WeakThis = TWeakPtr<FReplicationManagerState_Connected>(SharedThis(this)), Args](FConcertReplication_ChangeStream_Response&& Response)
+			.Next([WeakThis = TWeakPtr<FReplicationManagerState_Connected>(SharedThis(this)), Args, PredictedChange = MoveTemp(PredictedChange)](FConcertReplication_ChangeStream_Response&& Response)
 			{
 				Private::LogNetworkMessage(CVarLogStreamRequestsAndResponsesOnClient, Response);
 
@@ -242,7 +242,7 @@ namespace UE::ConcertSyncClient::Replication
 				else if (ThisPin && Response.ErrorCode == EReplicationResponseErrorCode::Timeout)
 				{
 					// HandleRemovingReplicatedObjects caused Request.ObjectsToRemove to stop being replicated. Revert.
-					ThisPin->RevertRemovingReplicatedObjects(Args);
+					ThisPin->RevertRemovingReplicatedObjects(PredictedChange);
 				}
 				
 				return FConcertReplication_ChangeStream_Response { MoveTemp(Response) };
@@ -415,32 +415,44 @@ namespace UE::ConcertSyncClient::Replication
 	{
 		static void ForEachObjectRemovedFromStreams(const FConcertReplication_ChangeStream_Request& Request, TFunctionRef<void(const FSoftObjectPath& ObjectPath, const TArray<FGuid>& Streams)> Callback)
 		{
-			TMap<FSoftObjectPath, TArray<FGuid>> BundledRemovedObjects;
-			for (const FConcertObjectInStreamID& RemovedObject : Request.ObjectsToRemove)
+			
+		}
+	}
+
+	TMap<FSoftObjectPath, TArray<FGuid>> FReplicationManagerState_Connected::HandleRemovingReplicatedObjects(const FConcertReplication_ChangeStream_Request& Request)
+	{
+		TMap<FSoftObjectPath, TArray<FGuid>> BundledRemovedObjects;
+		for (const FConcertObjectInStreamID& RemovedObject : Request.ObjectsToRemove)
+		{
+			BundledRemovedObjects.FindOrAdd(RemovedObject.Object).AddUnique(RemovedObject.StreamId);
+		}
+		for (const FConcertReplicationStream& Stream : RegisteredStreams)
+		{
+			const FGuid& StreamId = Stream.BaseDescription.Identifier;
+			if (!Request.StreamsToRemove.Contains(StreamId))
 			{
-				BundledRemovedObjects.FindOrAdd(RemovedObject.Object).Add(RemovedObject.StreamId);
+				continue;
 			}
-			for (const TPair<FSoftObjectPath, TArray<FGuid>>& RemovedObjectInfo : BundledRemovedObjects)
+
+			for (const TPair<FSoftObjectPath, FConcertReplicatedObjectInfo>& Pair : Stream.BaseDescription.ReplicationMap.ReplicatedObjects)
 			{
-				Callback(RemovedObjectInfo.Key, RemovedObjectInfo.Value);
+				BundledRemovedObjects.FindOrAdd(Pair.Key).AddUnique(StreamId);
 			}
 		}
-	}
-
-	void FReplicationManagerState_Connected::HandleRemovingReplicatedObjects(const FConcertReplication_ChangeStream_Request& Request)
-	{
-		Private::ForEachObjectRemovedFromStreams(Request, [this](const FSoftObjectPath& Object, const TArray<FGuid>& RemovedStreams)
+		
+		for (const TPair<FSoftObjectPath, TArray<FGuid>>& RemovedObjectInfo : BundledRemovedObjects)
 		{
-			ReplicationDataSource.RemoveReplicatedObjectStreams(Object, RemovedStreams);
-		});
+			ReplicationDataSource.RemoveReplicatedObjectStreams(RemovedObjectInfo.Key, RemovedObjectInfo.Value);
+		}
+		return BundledRemovedObjects;
 	}
 
-	void FReplicationManagerState_Connected::RevertRemovingReplicatedObjects(const FConcertReplication_ChangeStream_Request& Request)
+	void FReplicationManagerState_Connected::RevertRemovingReplicatedObjects(const TMap<FSoftObjectPath, TArray<FGuid>>& PredictedChange)
 	{
-		Private::ForEachObjectRemovedFromStreams(Request, [this](const FSoftObjectPath& Object, const TArray<FGuid>& RemovedStreams)
+		for (const TPair<FSoftObjectPath, TArray<FGuid>>& RemovedObjectInfo : PredictedChange)
 		{
-			ReplicationDataSource.AddReplicatedObjectStreams(Object, RemovedStreams);
-		});
+			ReplicationDataSource.AddReplicatedObjectStreams(RemovedObjectInfo.Key, RemovedObjectInfo.Value);
+		}
 	}
 
 	void FReplicationManagerState_Connected::UpdateReplicatedObjectsAfterAuthorityChange(FConcertReplication_ChangeAuthority_Request&& Request, const FConcertReplication_ChangeAuthority_Response& Response)
diff --git a/Engine/Plugins/Developer/Concert/ConcertSync/ConcertSyncClient/Source/ConcertSyncClient/Private/Replication/Manager/ReplicationManagerState_Connected.h b/Engine/Plugins/Developer/Concert/ConcertSync/ConcertSyncClient/Source/ConcertSyncClient/Private/Replication/Manager/ReplicationManagerState_Connected.h
--- a/Engine/Plugins/Developer/Concert/ConcertSync/ConcertSyncClient/Source/ConcertSyncClient/Private/Replication/Manager/ReplicationManagerState_Connected.h
+++ b/Engine/Plugins/Developer/Concert/ConcertSync/ConcertSyncClient/Source/ConcertSyncClient/Private/Replication/Manager/ReplicationManagerState_Connected.h
@@ -112,8 +112,8 @@ namespace UE::ConcertSyncClient::Replication
 		
 		/** Updates replicated objects affected by the change request. */
 		void UpdateReplicatedObjectsAfterStreamChange(const FConcertReplication_ChangeStream_Request& Request);
-		void HandleRemovingReplicatedObjects(const FConcertReplication_ChangeStream_Request& Request);
-		void RevertRemovingReplicatedObjects(const FConcertReplication_ChangeStream_Request& Request);
+		TMap<FSoftObjectPath, TArray<FGuid>> HandleRemovingReplicatedObjects(const FConcertReplication_ChangeStream_Request& Request);
+		void RevertRemovingReplicatedObjects(const TMap<FSoftObjectPath, TArray<FGuid>>& PredictedChange);
 
 		/**
 		 * Updates the objects which should be replicated after changing authority.
diff --git a/Engine/Plugins/Developer/Concert/ConcertSync/ConcertSyncClient/Source/ConcertSyncClient/Private/Replication/Processing/ClientReplicationDataCollector.h b/Engine/Plugins/Developer/Concert/ConcertSync/ConcertSyncClient/Source/ConcertSyncClient/Private/Replication/Processing/ClientReplicationDataCollector.h
--- a/Engine/Plugins/Developer/Concert/ConcertSync/ConcertSyncClient/Source/ConcertSyncClient/Private/Replication/Processing/ClientReplicationDataCollector.h
+++ b/Engine/Plugins/Developer/Concert/ConcertSync/ConcertSyncClient/Source/ConcertSyncClient/Private/Replication/Processing/ClientReplicationDataCollector.h
@@ -55,8 +55,8 @@ namespace UE::ConcertSyncClient::Replication
 		void AddReplicatedObjectStreams(const FSoftObjectPath& Object, TArrayView<const FGuid> AddedStreams);
 		/**
 		 * Indicates that certain properties of an object should no longer be replicated.
-		 * @param Object The object that should start replicating
-		 * @param RemovedStreams The streams determine which properties are to be replicated
+		 * @param Object The object that should is affected by the stream change
+		 * @param RemovedStreams The properties in these streams will no longer be replicated
 		 */
 		void RemoveReplicatedObjectStreams(const FSoftObjectPath& Object, TArrayView<const FGuid> RemovedStreams);
 		/**
diff --git a/Engine/Plugins/Developer/Concert/ConcertSync/ConcertSyncTest/Source/ConcertSyncTest/Private/Replication/Stream/StreamAuthorityInterplay.spec.cpp b/Engine/Plugins/Developer/Concert/ConcertSync/ConcertSyncTest/Source/ConcertSyncTest/Private/Replication/Stream/StreamAuthorityInterplay.spec.cpp
new file mode 100644
--- /dev/null
+++ b/Engine/Plugins/Developer/Concert/ConcertSync/ConcertSyncTest/Source/ConcertSyncTest/Private/Replication/Stream/StreamAuthorityInterplay.spec.cpp
@@ -0,0 +1,55 @@
+﻿// Copyright Epic Games, Inc. All Rights Reserved.
+
+#include "Replication/Util/Mocks/ReplicationWorkspaceCallInterceptorMock.h"
+#include "Replication/Util/Spec/ReplicationClient.h"
+#include "Replication/Util/Spec/ReplicationServer.h"
+#include "Replication/Util/Spec/ObjectTestReplicator.h"
+
+#include "Misc/AutomationTest.h"
+#include "Replication/ReplicationTestInterface.h"
+#include "Util/ScopedSessionDatabase.h"
+
+namespace UE::ConcertSyncTests::Replication::RestoreContent
+{
+	BEGIN_DEFINE_SPEC(FChangeStreamSpec, "Editor.Concert.Replication.Stream", EAutomationTestFlags::EditorContext | EAutomationTestFlags::EngineFilter)
+		TUniquePtr<FObjectTestReplicator> ObjectReplicator;
+	
+		TUniquePtr<FReplicationServer> Server;
+		FReplicationClient* Client = nullptr;
+
+		const FGuid StreamId = FGuid::NewGuid();
+	END_DEFINE_SPEC(FChangeStreamSpec);
+
+	/**
+	 * This tests that stream structure and authority correctly interact when issuing a FConcertReplication_ChangeStream_Request.
+	 * 
+	 * Old tests are in StreamRequestTests_x.cpp files of the same folder.
+	 * In the future, we want to move those over to the spec test format, too.
+	 */
+	void FChangeStreamSpec::Define()
+	{
+		BeforeEach([this]
+		{
+			ObjectReplicator = MakeUnique<FObjectTestReplicator>();
+			Server = MakeUnique<FReplicationServer>(*this);
+			Client = &Server->ConnectClient();
+			Client->JoinReplication();
+			
+		});
+		AfterEach([this]
+		{
+			Server.Reset();
+			ObjectReplicator.Reset();
+		});
+
+		It("Removing stream also removes authority locally", [this]
+		{
+			IConcertClientReplicationManager& ReplicationManager = Client->GetClientReplicationManager();
+			ReplicationManager.ChangeStream({ .StreamsToAdd = { ObjectReplicator->CreateStream(StreamId) } });
+			ReplicationManager.TakeAuthorityOver({ ObjectReplicator->TestObject });
+			ReplicationManager.ChangeStream({ .StreamsToRemove = { StreamId } });
+
+			TestEqual(TEXT("Clients updated local authority cache to 0"), ReplicationManager.GetClientOwnedObjects().Num(), 0);
+		});
+	}
+}
