Move presence to unreliable queue with non-unique messages. Added some additional information to UdpMessaging to help debugging.

diff --git a/Engine/Plugins/Developer/Concert/ConcertSync/ConcertSyncClient/Source/ConcertSyncClient/Private/ConcertClientPresenceMode.cpp b/Engine/Plugins/Developer/Concert/ConcertSync/ConcertSyncClient/Source/ConcertSyncClient/Private/ConcertClientPresenceMode.cpp
--- a/Engine/Plugins/Developer/Concert/ConcertSync/ConcertSyncClient/Source/ConcertSyncClient/Private/ConcertClientPresenceMode.cpp
+++ b/Engine/Plugins/Developer/Concert/ConcertSync/ConcertSyncClient/Source/ConcertSyncClient/Private/ConcertClientPresenceMode.cpp
@@ -77,7 +77,7 @@ void FConcertClientBasePresenceMode::SendEvents(IConcertClientSession& Session)
 
 	SetUpdateIndex(Session, FConcertClientPresenceDataUpdateEvent::StaticStruct()->GetFName(), PresenceDataUpdatedEvent);
 
-	Session.SendCustomEvent(PresenceDataUpdatedEvent, Session.GetSessionClientEndpointIds(), EConcertMessageFlags::UniqueId);
+	Session.SendCustomEvent(PresenceDataUpdatedEvent, Session.GetSessionClientEndpointIds(), EConcertMessageFlags::None);
 }
 
 FTransform FConcertClientBasePresenceMode::GetTransform()
@@ -195,7 +195,7 @@ void FConcertClientDesktopPresenceMode::SendEvents(IConcertClientSession& Sessio
 
 				SetUpdateIndex(Session, FConcertClientDesktopPresenceUpdateEvent::StaticStruct()->GetFName(), Event);
 
-				Session.SendCustomEvent(Event, Session.GetSessionClientEndpointIds(), EConcertMessageFlags::UniqueId);
+				Session.SendCustomEvent(Event, Session.GetSessionClientEndpointIds(), EConcertMessageFlags::None);
 			}
 		}
 	}
@@ -275,7 +275,7 @@ void FConcertClientVRPresenceMode::SendEvents(IConcertClientSession& Session)
 
 		SetUpdateIndex(Session, FConcertClientDesktopPresenceUpdateEvent::StaticStruct()->GetFName(), Event);
 
-		Session.SendCustomEvent(Event, Session.GetSessionClientEndpointIds(), EConcertMessageFlags::UniqueId);
+		Session.SendCustomEvent(Event, Session.GetSessionClientEndpointIds(), EConcertMessageFlags::None);
 	}
 }
 
diff --git a/Engine/Plugins/Messaging/UdpMessaging/Source/UdpMessaging/Private/Transport/UdpMessageSegmenter.cpp b/Engine/Plugins/Messaging/UdpMessaging/Source/UdpMessaging/Private/Transport/UdpMessageSegmenter.cpp
--- a/Engine/Plugins/Messaging/UdpMessaging/Source/UdpMessaging/Private/Transport/UdpMessageSegmenter.cpp
+++ b/Engine/Plugins/Messaging/UdpMessaging/Source/UdpMessaging/Private/Transport/UdpMessageSegmenter.cpp
@@ -255,7 +255,10 @@ bool FUdpMessageSegmenter::NeedSending(const FDateTime& CurrentTime)
 	const uint16 MaxNumResends = FMath::Clamp(CVarSegmenterMaxResends.GetValueOnAnyThread(), 1, 100);
 	if (AreAcknowledgementsComplete() == false && SentNumber > MaxNumResends)
 	{
-		UE_LOG(LogUdpMessaging, Warning, TEXT("Gave up sending with %d outstanding acks."), AcknowledgeSegments.Num() - AcknowledgeSegmentsCount);
+		UE_LOG(LogUdpMessaging, Warning, TEXT("Gave up sending with %d outstanding acks. Attempted %d tries."),
+			   AcknowledgeSegments.Num() - AcknowledgeSegmentsCount,
+			   MaxNumResends
+		);
 		SerializedMessage->UpdateState(EUdpSerializedMessageState::Invalid);
 		return false;
 	}
@@ -278,7 +281,10 @@ bool FUdpMessageSegmenter::NeedSending(const FDateTime& CurrentTime)
 		}
 		LastSentTime = CurrentTime;
 		++SentNumber;
-		UE_LOG(LogUdpMessaging, Warning, TEXT("Waiting for ack too long. Re-sending %d segments."), AcknowledgeSegments.Num() - AcknowledgeSegmentsCount);
+		UE_LOG(LogUdpMessaging, Warning, TEXT("Waiting for ack too long. Re-sending %d segments. Attempt %d out of %d tries."),
+			   AcknowledgeSegments.Num() - AcknowledgeSegmentsCount,
+			   SentNumber,
+			   MaxNumResends);
 		return true;
 	}
 
diff --git a/Engine/Plugins/Messaging/UdpMessaging/Source/UdpMessaging/Private/UdpMessagingModule.cpp b/Engine/Plugins/Messaging/UdpMessaging/Source/UdpMessaging/Private/UdpMessagingModule.cpp
--- a/Engine/Plugins/Messaging/UdpMessaging/Source/UdpMessaging/Private/UdpMessagingModule.cpp
+++ b/Engine/Plugins/Messaging/UdpMessaging/Source/UdpMessaging/Private/UdpMessagingModule.cpp
@@ -639,7 +639,7 @@ protected:
 			{
 				Settings->StaticEndpoints.AddUnique(CmdStaticEndpoint);
 			}
-
+			UE_LOG(LogUdpMessaging, Display, TEXT("Work queue size set to %d."), Settings->WorkQueueSize);
 			if (Settings->bShareKnownNodesWithActiveConnections)
 			{
 				UE_LOG(LogUdpMessaging, Display, TEXT("Sharing known nodes with endpoints is enabled."));
