Limit sequencer state events during playback.  We are already in a play state so no need to synchronize each sequencer time update.  This will reduce the number of messages that are sent.

diff --git a/Engine/Plugins/Developer/Concert/ConcertSync/ConcertSyncClient/Source/ConcertSyncClient/Private/ConcertClientSequencerManager.cpp b/Engine/Plugins/Developer/Concert/ConcertSync/ConcertSyncClient/Source/ConcertSyncClient/Private/ConcertClientSequencerManager.cpp
--- a/Engine/Plugins/Developer/Concert/ConcertSync/ConcertSyncClient/Source/ConcertSyncClient/Private/ConcertClientSequencerManager.cpp
+++ b/Engine/Plugins/Developer/Concert/ConcertSync/ConcertSyncClient/Source/ConcertSyncClient/Private/ConcertClientSequencerManager.cpp
@@ -103,10 +103,25 @@ public:
 			// If we have loop enabled and we have started back over again then we should force a send.
 			bShouldForceSend = InEvent.State.Time.AsSeconds() < CurrentStateEvent->State.Time.AsSeconds();
 		}
+		CheckForContinuationOfPlay(InEvent);
 		CurrentStateEvent = MoveTemp(InEvent);
 		Tick(Session, bShouldForceSend);
 	}
 
+	void CheckForContinuationOfPlay(const FConcertSequencerStateEvent& InEvent)
+	{
+		if (CurrentStateEvent)
+		{
+			const FConcertSequencerState& Current = CurrentStateEvent->State;
+			const FConcertSequencerState& Updated = InEvent.State;
+			bIsAContinuationOfPlay = Current.PlayerStatus == EConcertMovieScenePlayerStatus::Playing
+				&& Updated.PlayerStatus == EConcertMovieScenePlayerStatus::Playing
+				&& Updated.PlaybackRange == Current.PlaybackRange
+				&& Updated.PlaybackSpeed == Current.PlaybackSpeed
+				&& Updated.SequenceObjectPath == Current.SequenceObjectPath;
+		}
+	}
+
 	void Tick(TSharedPtr<IConcertClientSession>& Session, bool bForceSend = false)
 	{
 		if (CurrentStateEvent)
@@ -117,9 +132,14 @@ public:
 			{
 				bool bShouldSendReliably = bForceSend && CVarSequencerStatePacingEnabled.GetValueOnAnyThread();
 				EConcertMessageFlags Flags = bShouldSendReliably || CVarSendStateEventsAsUnreliable.GetValueOnAnyThread() == 0 ? EConcertMessageFlags::ReliableOrdered : EConcertMessageFlags::None;
-				Session->SendCustomEvent(*CurrentStateEvent, Session->GetSessionServerEndpointId(), Flags);
-				LastSendTime = CurrentTime;
-				CurrentStateEvent.Reset();
+				if (bForceSend || !bIsAContinuationOfPlay)
+				{
+					Session->SendCustomEvent(*CurrentStateEvent, Session->GetSessionServerEndpointId(), Flags);
+					LastSendTime = CurrentTime;
+					CurrentStateEvent.Reset();
+					bIsFirstEvent = false;
+					bIsAContinuationOfPlay = false;
+				}
 			}
 		}
 	}
@@ -128,6 +148,7 @@ private:
 	TOptional<FConcertSequencerStateEvent> CurrentStateEvent;
 	double LastSendTime = 0;
 	bool bIsFirstEvent = true;
+	bool bIsAContinuationOfPlay = false;
 };
 
 class FConcertClientSequencePreloader : public TSharedFromThis<FConcertClientSequencePreloader>
