Extend IConcertClientWorkspace with FindReplicationEvent.

diff --git a/Engine/Plugins/Developer/Concert/ConcertSync/ConcertSyncClient/Source/ConcertSyncClient/Private/ConcertClientWorkspace.cpp b/Engine/Plugins/Developer/Concert/ConcertSync/ConcertSyncClient/Source/ConcertSyncClient/Private/ConcertClientWorkspace.cpp
--- a/Engine/Plugins/Developer/Concert/ConcertSync/ConcertSyncClient/Source/ConcertSyncClient/Private/ConcertClientWorkspace.cpp
+++ b/Engine/Plugins/Developer/Concert/ConcertSync/ConcertSyncClient/Source/ConcertSyncClient/Private/ConcertClientWorkspace.cpp
@@ -22,6 +22,7 @@
 #include "ConcertClientLiveTransactionAuthors.h"
 #include "IConcertSession.h"
 #include "IConcertFileSharingService.h"
+#include "Replication/Messages/ReplicationActivity.h"
 
 #include "Algo/AllOf.h"
 #include "Containers/ArrayBuilder.h"
@@ -43,7 +44,6 @@
 #include "AssetRegistry/AssetRegistryModule.h"
 #include "Engine/World.h"
 #include "Engine/Engine.h"
-#include "Replication/Messages/ReplicationActivity.h"
 
 #if WITH_EDITOR
 	#include "Editor.h"
@@ -446,6 +446,11 @@ bool FConcertClientWorkspace::FindPackageEvent(const int64 PackageEventId, FConc
 	return LiveSession->GetSessionDatabase().GetPackageEventMetaData(PackageEventId, OutPackageEvent.PackageRevision, OutPackageEvent.PackageInfo);
 }
 
+bool FConcertClientWorkspace::FindReplicationEvent(const int64 ReplicationEventId, FConcertSyncReplicationEvent& OutReplicationEvent) const
+{
+	return LiveSession->GetSessionDatabase().GetReplicationEvent(ReplicationEventId, OutReplicationEvent);
+}
+
 bool FConcertClientWorkspace::IsTransactionEventPartiallySynced(const FConcertSyncTransactionEvent& TransactionEvent) const
 {
 	return TransactionEvent.Transaction.ExportedObjects.Num() == 0;
diff --git a/Engine/Plugins/Developer/Concert/ConcertSync/ConcertSyncClient/Source/ConcertSyncClient/Private/ConcertClientWorkspace.h b/Engine/Plugins/Developer/Concert/ConcertSync/ConcertSyncClient/Source/ConcertSyncClient/Private/ConcertClientWorkspace.h
--- a/Engine/Plugins/Developer/Concert/ConcertSync/ConcertSyncClient/Source/ConcertSyncClient/Private/ConcertClientWorkspace.h
+++ b/Engine/Plugins/Developer/Concert/ConcertSync/ConcertSyncClient/Source/ConcertSyncClient/Private/ConcertClientWorkspace.h
@@ -63,6 +63,7 @@ public:
 	virtual bool FindTransactionEvent(const int64 TransactionEventId, FConcertSyncTransactionEvent& OutTransactionEvent, const bool bMetaDataOnly) const override;
 	virtual TFuture<TOptional<FConcertSyncTransactionEvent>> FindOrRequestTransactionEvent(const int64 TransactionEventId, const bool bMetaDataOnly) override;
 	virtual bool FindPackageEvent(const int64 PackageEventId, FConcertSyncPackageEventMetaData& OutPackageEvent) const override;
+	virtual bool FindReplicationEvent(const int64 ReplicationEventId, FConcertSyncReplicationEvent& OutReplicationEvent) const override;
 	virtual void GetActivities(const int64 FirstActivityIdToFetch, const int64 MaxNumActivities, TMap<FGuid, FConcertClientInfo>& OutEndpointClientInfoMap, TArray<FConcertSessionActivity>& OutActivities) const override;
 	virtual int64 GetLastActivityId() const override;
 	virtual FOnActivityAddedOrUpdated& OnActivityAddedOrUpdated() override;
diff --git a/Engine/Plugins/Developer/Concert/ConcertSync/ConcertSyncClient/Source/ConcertSyncClient/Public/IConcertClientWorkspace.h b/Engine/Plugins/Developer/Concert/ConcertSync/ConcertSyncClient/Source/ConcertSyncClient/Public/IConcertClientWorkspace.h
--- a/Engine/Plugins/Developer/Concert/ConcertSync/ConcertSyncClient/Source/ConcertSyncClient/Public/IConcertClientWorkspace.h
+++ b/Engine/Plugins/Developer/Concert/ConcertSync/ConcertSyncClient/Source/ConcertSyncClient/Public/IConcertClientWorkspace.h
@@ -14,6 +14,8 @@ class ISourceControlProvider;
 class IConcertClientSession;
 class IConcertClientDataStore;
 
+struct FConcertSyncReplicationEvent;
+
 DECLARE_DELEGATE_RetVal(bool, FCanFinalizeWorkspaceDelegate);
 DECLARE_DELEGATE_RetVal(bool, FCanProcessPendingPackages);
 
@@ -118,7 +120,7 @@ public:
 	 * @return true if dirty event should be ignored for said package.
 	 */
 	virtual bool ShouldIgnorePackageDirtyEvent(class UPackage* InPackage) const = 0;
-
+	
 	/**
 	 * Lookup the specified transaction event.
 	 * @param[in] TransactionEventId ID of the transaction to look for.
@@ -150,6 +152,14 @@ public:
 	 */
 	virtual bool FindPackageEvent(const int64 PackageEventId, FConcertSyncPackageEventMetaData& OutPackageEvent) const = 0;
 
+	/**
+	 * Lookup the specified replication event.
+	 * @param[in] ReplicationEventId ID of the replication event to look for.
+	 * @param[out] OutReplicationEvent The replication event corresponding to ReplicationEventId if found.
+	 * @return Whether the replication event was found and requested data available. 
+	 */
+	virtual bool FindReplicationEvent(const int64 ReplicationEventId, FConcertSyncReplicationEvent& OutReplicationEvent) const = 0;
+
 	/**
 	 * @return the delegate called every time the workspace is synced. This is when the server has indicated that all
 	 * activities have been sent and ready to be finalized in the current workspace.
diff --git a/Engine/Plugins/Developer/Concert/ConcertSync/ConcertSyncTest/Source/ConcertSyncTest/Private/Util/Mock/ConcertClientWorkspaceBaseMock.h b/Engine/Plugins/Developer/Concert/ConcertSync/ConcertSyncTest/Source/ConcertSyncTest/Private/Util/Mock/ConcertClientWorkspaceBaseMock.h
--- a/Engine/Plugins/Developer/Concert/ConcertSync/ConcertSyncTest/Source/ConcertSyncTest/Private/Util/Mock/ConcertClientWorkspaceBaseMock.h
+++ b/Engine/Plugins/Developer/Concert/ConcertSync/ConcertSyncTest/Source/ConcertSyncTest/Private/Util/Mock/ConcertClientWorkspaceBaseMock.h
@@ -33,6 +33,7 @@ namespace UE::ConcertSyncTests::Replication
 		virtual bool FindTransactionEvent(const int64 TransactionEventId, FConcertSyncTransactionEvent& OutTransactionEvent, const bool bMetaDataOnly) const override { return NotMocked<bool>(); }
 		virtual TFuture<TOptional<FConcertSyncTransactionEvent>> FindOrRequestTransactionEvent(const int64 TransactionEventId, const bool bMetaDataOnly) override { return NotMocked<TFuture<TOptional<FConcertSyncTransactionEvent>>>(); }
 		virtual bool FindPackageEvent(const int64 PackageEventId, FConcertSyncPackageEventMetaData& OutPackageEvent) const override { return NotMocked<bool>(); }
+		virtual bool FindReplicationEvent(const int64 ReplicationEventId, FConcertSyncReplicationEvent& OutReplicationEvent) const override { return NotMocked<bool>(); }
 		virtual FOnWorkspaceSynchronized& OnWorkspaceSynchronized() override { return NotMocked<FOnWorkspaceSynchronized&>(OnWorkspaceSynchronizedDelegate); }
 		virtual FOnFinalizeWorkspaceSyncCompleted& OnFinalizeWorkspaceSyncCompleted() override { return NotMocked<FOnFinalizeWorkspaceSyncCompleted&>(OnFinalizeWorkspaceSyncCompletedDelegate); }
 		virtual void AddWorkspaceFinalizeDelegate(FName InDelegateName, FCanFinalizeWorkspaceDelegate InDelegate) override { NotMocked<void>(); }
