Added Concert support for duplicating an object

* Added FConcertSerializedObjectData::SourceObject which denotes that this object was created by copying another.
* Updated ConcertSyncClientUtil::GetObject to handle duplication.
    * Note: Actor duplication is only currently implemented for the editor.
* Updated UConcertClientObjectFactory to support custom object duplication.

Note: Concert currently doesn't use this workflow internally, but anything generating a custom FConcertSerializedObjectData is able to set the SourceObject data.

#rnx

diff --git a/Engine/Plugins/Developer/Concert/ConcertSync/ConcertSyncClient/Source/ConcertSyncClient/Private/ConcertClientTransactionBridge.cpp b/Engine/Plugins/Developer/Concert/ConcertSync/ConcertSyncClient/Source/ConcertSyncClient/Private/ConcertClientTransactionBridge.cpp
--- a/Engine/Plugins/Developer/Concert/ConcertSync/ConcertSyncClient/Source/ConcertSyncClient/Private/ConcertClientTransactionBridge.cpp
+++ b/Engine/Plugins/Developer/Concert/ConcertSync/ConcertSyncClient/Source/ConcertSyncClient/Private/ConcertClientTransactionBridge.cpp
@@ -388,7 +388,7 @@ void ProcessTransactionEvent(const FConcertTransactionEventBase& InEvent, const
 			}
 
 			// Find or create the object
-			TransactionObjectRef = ConcertSyncClientUtil::GetObject(ObjectUpdate.ObjectId, ObjectUpdate.ObjectData.NewName, ObjectUpdate.ObjectData.NewOuterPathName, ObjectUpdate.ObjectData.NewExternalPackageName, ObjectUpdate.ObjectData.bAllowCreate);
+			TransactionObjectRef = ConcertSyncClientUtil::GetObject(ObjectUpdate.ObjectId, ObjectUpdate.ObjectData.NewName, ObjectUpdate.ObjectData.NewOuterPathName, ObjectUpdate.ObjectData.NewExternalPackageName, ObjectUpdate.ObjectData.SourceObject, ObjectUpdate.ObjectData.bAllowCreate);
 			bObjectsDeleted |= (ObjectUpdate.ObjectData.bIsPendingKill || TransactionObjectRef.NeedsGC());
 
 			if (TransactionObjectRef.Obj)
@@ -475,7 +475,7 @@ void ProcessTransactionEvent(const FConcertTransactionEventBase& InEvent, const
 	}
 
 #if WITH_EDITOR
-	UObject* PrimaryObject = InEvent.PrimaryObjectId.ObjectName.IsNone() ? nullptr : ConcertSyncClientUtil::GetObject(InEvent.PrimaryObjectId, FName(), FName(), FName(), /*bAllowCreate*/false).Obj;
+	UObject* PrimaryObject = InEvent.PrimaryObjectId.ObjectName.IsNone() ? nullptr : ConcertSyncClientUtil::GetObject(InEvent.PrimaryObjectId, FName(), FName(), FName(), FSoftObjectPath(), /*bAllowCreate*/false).Obj;
 	FEditorTransactionNotification EditorTransactionNotification(FTransactionContext(InEvent.TransactionId, InEvent.OperationId, LOCTEXT("ConcertTransactionEvent", "Concert Transaction Event"), TEXT("Concert Transaction Event"), PrimaryObject));
 	if (!bIsSnapshot)
 	{
diff --git a/Engine/Plugins/Developer/Concert/ConcertSync/ConcertSyncClient/Source/ConcertSyncClient/Private/ConcertClientTransactionManager.cpp b/Engine/Plugins/Developer/Concert/ConcertSync/ConcertSyncClient/Source/ConcertSyncClient/Private/ConcertClientTransactionManager.cpp
--- a/Engine/Plugins/Developer/Concert/ConcertSync/ConcertSyncClient/Source/ConcertSyncClient/Private/ConcertClientTransactionManager.cpp
+++ b/Engine/Plugins/Developer/Concert/ConcertSync/ConcertSyncClient/Source/ConcertSyncClient/Private/ConcertClientTransactionManager.cpp
@@ -434,7 +434,7 @@ void MakeObjectUniqueToAvoidCollision(UObject* InObject)
  */
 FRenameObjectResult RenameObjectsToAvoidCollision(const FConcertObjectId& InObjectIdToRename)
 {
-	UObject* ObjectToRename = ConcertSyncClientUtil::GetObject(InObjectIdToRename, FName(), FName(), FName(), false).Obj;
+	UObject* ObjectToRename = ConcertSyncClientUtil::GetObject(InObjectIdToRename, FName(), FName(), FName(), FSoftObjectPath(), false).Obj;
 	// The object does not exist or is marked for GC.
 	if (!IsValid(ObjectToRename))
 	{
diff --git a/Engine/Plugins/Developer/Concert/ConcertSync/ConcertSyncClient/Source/ConcertSyncClient/Private/ConcertSyncClientUtil.cpp b/Engine/Plugins/Developer/Concert/ConcertSync/ConcertSyncClient/Source/ConcertSyncClient/Private/ConcertSyncClientUtil.cpp
--- a/Engine/Plugins/Developer/Concert/ConcertSync/ConcertSyncClient/Source/ConcertSyncClient/Private/ConcertSyncClientUtil.cpp
+++ b/Engine/Plugins/Developer/Concert/ConcertSync/ConcertSyncClient/Source/ConcertSyncClient/Private/ConcertSyncClientUtil.cpp
@@ -44,6 +44,7 @@
 	#include "PackageTools.h"
 	#include "Selection.h"
 	#include "Subsystems/AssetEditorSubsystem.h"
+	#include "Subsystems/EditorActorSubsystem.h"
 	#include "UnrealEdGlobals.h"
 #endif
 
@@ -206,7 +207,7 @@ int32 GetObjectPathDepth(UObject* InObjToTest)
 	return Depth;
 }
 
-FGetObjectResult GetObject(const FConcertObjectId& InObjectId, const FName InNewName, const FName InNewOuterPath, const FName InNewPackageName, const bool bAllowCreate)
+FGetObjectResult GetObject(const FConcertObjectId& InObjectId, const FName InNewName, const FName InNewOuterPath, const FName InNewPackageName, const FSoftObjectPath& InSourceObject, const bool bAllowCreate)
 {
 	const bool bIsRename = !InNewName.IsNone();
 	const bool bIsOuterChange = !InNewOuterPath.IsNone();
@@ -350,6 +351,18 @@ FGetObjectResult GetObject(const FConcertObjectId& InObjectId, const FName InNew
 
 		if (bAllowCreate)
 		{
+			UObject* SourceObject = nullptr;
+			if (!InSourceObject.IsNull())
+			{
+				SourceObject = InSourceObject.ResolveObject();
+				UE_CLOG(!SourceObject, LogConcert, Warning, TEXT("Failed to find source object '%s' for '%s'. This object will be created from its CDO instead."), *InSourceObject.ToString(), *ObjectNameToCreate.ToString());
+			}
+			if (SourceObject && SourceObject->GetClass() != ObjectClass)
+			{
+				UE_LOG(LogConcert, Warning, TEXT("Discarding source object '%s' for '%s' as it was not the expected class (%s). This object will be created from its CDO instead."), *InSourceObject.ToString(), *ObjectNameToCreate.ToString(), *ObjectClass->GetPathName());
+				SourceObject = nullptr;
+			}
+
 			FGetObjectResult ObjectResult;
 			ObjectResult.Factory = Factory;
 
@@ -357,7 +370,14 @@ FGetObjectResult GetObject(const FConcertObjectId& InObjectId, const FName InNew
 			bool bFactoryHandledCreation = false;
 			if (Factory)
 			{
-				bFactoryHandledCreation = Factory->CreateObject(ObjectResult.Obj, NewObjectOuter, ObjectClass, *ObjectNameToCreate.ToString(), (EObjectFlags)InObjectId.ObjectPersistentFlags);
+				if (SourceObject)
+				{
+					bFactoryHandledCreation = Factory->DuplicateObject(ObjectResult.Obj, SourceObject, NewObjectOuter, ObjectClass, *ObjectNameToCreate.ToString(), (EObjectFlags)InObjectId.ObjectPersistentFlags);
+				}
+				else
+				{
+					bFactoryHandledCreation = Factory->CreateObject(ObjectResult.Obj, NewObjectOuter, ObjectClass, *ObjectNameToCreate.ToString(), (EObjectFlags)InObjectId.ObjectPersistentFlags);
+				}
 			}
 			if (!bFactoryHandledCreation)
 			{
@@ -377,13 +397,38 @@ FGetObjectResult GetObject(const FConcertObjectId& InObjectId, const FName InNew
 							UObject* ExistingObjectOfDifferentClass = StaticFindObjectFast(nullptr, OuterLevel, ObjectNameToCreate);
 							if (!ExistingObjectOfDifferentClass)
 							{
-								FActorSpawnParameters SpawnParams;
-								SpawnParams.Name = ObjectNameToCreate;
-								SpawnParams.OverrideLevel = OuterLevel;
-								SpawnParams.SpawnCollisionHandlingOverride = ESpawnActorCollisionHandlingMethod::AlwaysSpawn;
-								SpawnParams.bNoFail = true;
-								SpawnParams.ObjectFlags = (EObjectFlags)InObjectId.ObjectPersistentFlags;
-								ObjectResult.Obj = OwnerWorld->SpawnActor<AActor>(ObjectClass, FVector::ZeroVector, FRotator::ZeroRotator, SpawnParams);
+								if (SourceObject)
+								{
+#if WITH_EDITOR
+									if (GIsEditor && !OwnerWorld->IsPlayInEditor())
+									{
+										UEditorActorSubsystem::FActorDuplicateParameters DuplicateParams;
+										DuplicateParams.LevelOverride = OuterLevel;
+										DuplicateParams.bTransact = false;
+										ObjectResult.Obj = GUnrealEd->GetEditorSubsystem<UEditorActorSubsystem>()->DuplicateActor(CastChecked<AActor>(SourceObject), OwnerWorld, FVector::ZeroVector, DuplicateParams);
+										if (ObjectResult.Obj)
+										{
+											ObjectResult.Obj->SetFlags((EObjectFlags)InObjectId.ObjectPersistentFlags);
+											ObjectResult.Obj->Rename(*ObjectNameToCreate.ToString(), nullptr, REN_NonTransactional | REN_DoNotDirty);
+										}
+									}
+									else
+#endif	// WITH_EDITOR
+									{
+										// TODO: There is no direct equivalent of DuplicateActor for runtime use, though DuplicateObject *might* work. There is currently no use-case for this, so revisit a solution if it becomes an issue
+										checkf(false, TEXT("Duplicating an actor outside of the editor is not currently supported!"));
+									}
+								}
+								else
+								{
+									FActorSpawnParameters SpawnParams;
+									SpawnParams.Name = ObjectNameToCreate;
+									SpawnParams.OverrideLevel = OuterLevel;
+									SpawnParams.SpawnCollisionHandlingOverride = ESpawnActorCollisionHandlingMethod::AlwaysSpawn;
+									SpawnParams.bNoFail = true;
+									SpawnParams.ObjectFlags = (EObjectFlags)InObjectId.ObjectPersistentFlags;
+									ObjectResult.Obj = OwnerWorld->SpawnActor<AActor>(ObjectClass, FVector::ZeroVector, FRotator::ZeroRotator, SpawnParams);
+								}
 							}
 							else
 							{
@@ -403,7 +448,18 @@ FGetObjectResult GetObject(const FConcertObjectId& InObjectId, const FName InNew
 				}
 				else
 				{
-					ObjectResult.Obj = NewObject<UObject>(NewObjectOuter, ObjectClass, *ObjectNameToCreate.ToString(), (EObjectFlags)InObjectId.ObjectPersistentFlags);
+					if (SourceObject)
+					{
+						ObjectResult.Obj = DuplicateObject<UObject>(SourceObject, NewObjectOuter, *ObjectNameToCreate.ToString());
+						if (ObjectResult.Obj)
+						{
+							ObjectResult.Obj->SetFlags((EObjectFlags)InObjectId.ObjectPersistentFlags);
+						}
+					}
+					else
+					{
+						ObjectResult.Obj = NewObject<UObject>(NewObjectOuter, ObjectClass, *ObjectNameToCreate.ToString(), (EObjectFlags)InObjectId.ObjectPersistentFlags);
+					}
 
 					if (UActorComponent* NewComponent = Cast<UActorComponent>(ObjectResult.Obj))
 					{
diff --git a/Engine/Plugins/Developer/Concert/ConcertSync/ConcertSyncClient/Source/ConcertSyncClient/Private/ConcertSyncClientUtil.h b/Engine/Plugins/Developer/Concert/ConcertSync/ConcertSyncClient/Source/ConcertSyncClient/Private/ConcertSyncClientUtil.h
--- a/Engine/Plugins/Developer/Concert/ConcertSync/ConcertSyncClient/Source/ConcertSyncClient/Private/ConcertSyncClientUtil.h
+++ b/Engine/Plugins/Developer/Concert/ConcertSync/ConcertSyncClient/Source/ConcertSyncClient/Private/ConcertSyncClientUtil.h
@@ -86,7 +86,7 @@ namespace ConcertSyncClientUtil
 
 	int32 GetObjectPathDepth(UObject* InObjToTest);
 
-	FGetObjectResult GetObject(const FConcertObjectId& InObjectId, const FName InNewName, const FName InNewOuterPath, const FName InNewPackageName, const bool bAllowCreate);
+	FGetObjectResult GetObject(const FConcertObjectId& InObjectId, const FName InNewName, const FName InNewOuterPath, const FName InNewPackageName, const FSoftObjectPath& InSourceObject, const bool bAllowCreate);
 	
 	TArray<const FProperty*> GetExportedProperties(const UStruct* InStruct, const TArray<FName>& InPropertyNames, const bool InIncludeEditorOnlyData);
 
diff --git a/Engine/Plugins/Developer/Concert/ConcertSync/ConcertSyncClient/Source/ConcertSyncClient/Public/ConcertClientObjectFactory.h b/Engine/Plugins/Developer/Concert/ConcertSync/ConcertSyncClient/Source/ConcertSyncClient/Public/ConcertClientObjectFactory.h
--- a/Engine/Plugins/Developer/Concert/ConcertSync/ConcertSyncClient/Source/ConcertSyncClient/Public/ConcertClientObjectFactory.h
+++ b/Engine/Plugins/Developer/Concert/ConcertSync/ConcertSyncClient/Source/ConcertSyncClient/Public/ConcertClientObjectFactory.h
@@ -53,9 +53,18 @@ public:
 		return false;
 	}
 	
+	/**
+	 * Attempt to create a new object from the given source object, and fill in OutObject.
+	 * @return True if this function handled the creation attempt (this doesn't mean that OutObject isn't null!), or false if we should fallback to using StaticDuplicateObject.
+	 */
+	virtual bool DuplicateObject(UObject*& OutObject, const UObject* SourceObject, UObject* Outer, const UClass* Class, const FName Name, const EObjectFlags Flags) const
+	{
+		return false;
+	}
+
 	/**
 	 * Perform any additional initialization on the given objects.
-	 * @note This is called after CreateObject has been called for every object in the transaction, and can be used for delayed initialization of interdependent objects.
+	 * @note This is called after CreateObject/DuplicateObject has been called for every object in the transaction, and can be used for delayed initialization of interdependent objects.
 	 */
 	virtual void InitializeObjects(TArrayView<UObject* const> Objects) const
 	{
diff --git a/Engine/Plugins/Developer/Concert/ConcertSync/ConcertSyncCore/Source/ConcertSyncCore/Public/ConcertTransactionEvents.h b/Engine/Plugins/Developer/Concert/ConcertSync/ConcertSyncCore/Source/ConcertSyncCore/Public/ConcertTransactionEvents.h
--- a/Engine/Plugins/Developer/Concert/ConcertSync/ConcertSyncCore/Source/ConcertSyncCore/Public/ConcertTransactionEvents.h
+++ b/Engine/Plugins/Developer/Concert/ConcertSync/ConcertSyncCore/Source/ConcertSyncCore/Public/ConcertTransactionEvents.h
@@ -128,6 +128,9 @@ struct FConcertSerializedObjectData
 	UPROPERTY()
 	FName NewExternalPackageName;
 
+	UPROPERTY()
+	FSoftObjectPath SourceObject;
+
 	UPROPERTY()
 	TArray<uint8> SerializedData;
 };
