Add support for primitive component in Level Instances / PA transactions.

diff --git a/Engine/Plugins/Developer/Concert/ConcertSync/ConcertSyncClient/Source/ConcertSyncClient/Private/ConcertClientTransactionBridge.cpp b/Engine/Plugins/Developer/Concert/ConcertSync/ConcertSyncClient/Source/ConcertSyncClient/Private/ConcertClientTransactionBridge.cpp
--- a/Engine/Plugins/Developer/Concert/ConcertSync/ConcertSyncClient/Source/ConcertSyncClient/Private/ConcertClientTransactionBridge.cpp
+++ b/Engine/Plugins/Developer/Concert/ConcertSync/ConcertSyncClient/Source/ConcertSyncClient/Private/ConcertClientTransactionBridge.cpp
@@ -17,6 +17,7 @@
 #include "UObject/Package.h"
 #include "LevelInstance/LevelInstanceActor.h"
 #include "LevelInstance/LevelInstanceComponent.h"
+#include "Components/PrimitiveComponent.h"
 
 #if WITH_EDITOR
 	#include "Editor.h"
@@ -899,6 +900,10 @@ UClass* GetModifiedClass(UObject* InObject)
 	{
 		return USceneComponent::StaticClass();
 	}
+	if (Cast<UPrimitiveComponent>(InObject))
+	{
+		return UPrimitiveComponent::StaticClass();
+	}
 	return nullptr;
 }
 
