Add a project setting to control the visibilty of presence in game mode.

diff --git a/Engine/Plugins/Developer/Concert/ConcertSync/ConcertSyncClient/Source/ConcertSyncClient/Private/ConcertClientPresenceActor.cpp b/Engine/Plugins/Developer/Concert/ConcertSync/ConcertSyncClient/Source/ConcertSyncClient/Private/ConcertClientPresenceActor.cpp
--- a/Engine/Plugins/Developer/Concert/ConcertSync/ConcertSyncClient/Source/ConcertSyncClient/Private/ConcertClientPresenceActor.cpp
+++ b/Engine/Plugins/Developer/Concert/ConcertSync/ConcertSyncClient/Source/ConcertSyncClient/Private/ConcertClientPresenceActor.cpp
@@ -7,6 +7,7 @@
 #include "UObject/StructOnScope.h"
 #include "Engine/Scene.h"
 #include "ConcertPresenceEvents.h"
+#include "ConcertSyncSettings.h"
 #include "Components/StaticMeshComponent.h"
 #include "Components/TextRenderComponent.h"
 #include "ConcertAssetContainer.h"
@@ -31,9 +32,9 @@ AConcertClientPresenceActor::AConcertClientPresenceActor(const FObjectInitialize
 	: Super(ObjectInitializer)
 {
 	// Initialize if this actor will be editor only based the `EnablePresenceInGame` console variable.
-	bIsEditorOnlyActor = CVarEnablePresenceInGame.GetValueOnAnyThread() == 0;
+	UpdateVisibleInGameState();
 
-	// Set root component 
+	// Set root component
 	{
 		USceneComponent* SceneRootComponent = CreateDefaultSubobject<USceneComponent>(TEXT("Root"));
 		AddOwnedComponent(SceneRootComponent);
@@ -64,6 +65,22 @@ bool AConcertClientPresenceActor::ShouldTickIfViewportsOnly() const
 	return true;
 }
 
+void AConcertClientPresenceActor::UpdateVisibleInGameState()
+{
+	EShowPresenceInGame PresenceInGame = GetDefault<UConcertSyncConfig>()->ShowPresenceInGame;
+	bIsEditorOnlyActor = CVarEnablePresenceInGame.GetValueOnAnyThread() == 0;
+
+	if (bIsEditorOnlyActor && GIsEditor &&
+		(PresenceInGame == EShowPresenceInGame::VisibleInGameOnlyEditor || PresenceInGame == EShowPresenceInGame::VisibleInGameAlways))
+	{
+		bIsEditorOnlyActor = false;
+	}
+	else if (!GIsEditor && bIsEditorOnlyActor && PresenceInGame == EShowPresenceInGame::VisibleInGameAlways)
+	{
+		bIsEditorOnlyActor = false;
+	}
+}
+
 bool AConcertClientPresenceActor::IsEditorOnly() const
 {
 	return bIsEditorOnlyActor;
diff --git a/Engine/Plugins/Developer/Concert/ConcertSync/ConcertSyncClient/Source/ConcertSyncClient/Public/ConcertClientPresenceActor.h b/Engine/Plugins/Developer/Concert/ConcertSync/ConcertSyncClient/Source/ConcertSyncClient/Public/ConcertClientPresenceActor.h
--- a/Engine/Plugins/Developer/Concert/ConcertSync/ConcertSyncClient/Source/ConcertSyncClient/Public/ConcertClientPresenceActor.h
+++ b/Engine/Plugins/Developer/Concert/ConcertSync/ConcertSyncClient/Source/ConcertSyncClient/Public/ConcertClientPresenceActor.h
@@ -40,6 +40,7 @@ public:
 		return false;
 	}
 #endif
+
 	/**
 	 * Set the presence name.
 	 * Called by the presence manager to set the display name of the client on the actor.
@@ -95,5 +96,9 @@ protected:
 
 	/** Movement object to interpolate presence movement between update events.*/
 	TOptional<FConcertClientMovement> PresenceMovement;
+
+private:
+	void UpdateVisibleInGameState();
+
 };
 
diff --git a/Engine/Plugins/Developer/Concert/ConcertSync/ConcertSyncCore/Source/ConcertSyncCore/Public/ConcertSyncSettings.h b/Engine/Plugins/Developer/Concert/ConcertSync/ConcertSyncCore/Source/ConcertSyncCore/Public/ConcertSyncSettings.h
--- a/Engine/Plugins/Developer/Concert/ConcertSync/ConcertSyncCore/Source/ConcertSyncCore/Public/ConcertSyncSettings.h
+++ b/Engine/Plugins/Developer/Concert/ConcertSync/ConcertSyncCore/Source/ConcertSyncCore/Public/ConcertSyncSettings.h
@@ -41,6 +41,18 @@ struct CONCERTSYNCCORE_API FPackageClassFilter
 	TArray<FString> ContentPaths;
 };
 
+UENUM()
+enum class EShowPresenceInGame
+{
+	/** The presence actor is not visible in game mode. */
+	NotVisibleInGame,
+
+	/** The presence actor is visible in game only in editor sessions. */
+	VisibleInGameOnlyEditor,
+
+	/** The presence actor is always visible in game including -game modes. This is equivalent to Concert.EnablePresenceInGame */
+	VisibleInGameAlways
+};
 
 UCLASS(config=ConcertSyncCore)
 class CONCERTSYNCCORE_API UConcertSyncConfig : public UObject
@@ -68,6 +80,12 @@ public:
 	UPROPERTY(config, EditAnywhere, AdvancedDisplay, Category="Transaction Settings")
 	bool bShowPresenceInPIE;
 
+	/**
+	 * Should we show presence when we are in game mode when running the editor.
+	 */
+	UPROPERTY(config, EditAnywhere, AdvancedDisplay, Category="Transaction Settings")
+	EShowPresenceInGame ShowPresenceInGame = EShowPresenceInGame::NotVisibleInGame;
+
 	/**
 	 * Should we include object changes that have been generated via a transaction annotation
 	 * (where possible), or should we send the entire transaction annotation blob instead.
