Fixed duplicated objects in Concert having an incorrect transform when they were attached to other objects

#rnx

diff --git a/Engine/Plugins/Developer/Concert/ConcertSync/ConcertSyncClient/Source/ConcertSyncClient/Private/ConcertSyncClientUtil.cpp b/Engine/Plugins/Developer/Concert/ConcertSync/ConcertSyncClient/Source/ConcertSyncClient/Private/ConcertSyncClientUtil.cpp
--- a/Engine/Plugins/Developer/Concert/ConcertSync/ConcertSyncClient/Source/ConcertSyncClient/Private/ConcertSyncClientUtil.cpp
+++ b/Engine/Plugins/Developer/Concert/ConcertSync/ConcertSyncClient/Source/ConcertSyncClient/Private/ConcertSyncClientUtil.cpp
@@ -410,6 +410,12 @@ FGetObjectResult GetObject(const FConcertObjectId& InObjectId, const FName InNew
 										{
 											ObjectResult.Obj->SetFlags((EObjectFlags)InObjectId.ObjectPersistentFlags);
 											ObjectResult.Obj->Rename(*ObjectNameToCreate.ToString(), nullptr, REN_NonTransactional | REN_DoNotDirty);
+
+											// Clear any attachment, as Concert assumes an actor starts detached
+											// This avoids issues with the transform being offset from the attached root rather than the world origin
+											FDetachmentTransformRules DetachmentRules = FDetachmentTransformRules::KeepWorldTransform;
+											DetachmentRules.bCallModify = false;
+											CastChecked<AActor>(ObjectResult.Obj)->DetachFromActor(DetachmentRules);
 										}
 									}
 									else
