Allow take recorder to disable sequencer pacing which is necessary to have the pending take update playhead position during a record.  We also had to reorder when the events get applied because the sequencer events were setting the UMovieScene range before the time adjustment event could see the delta time change so that it could update the section position for playback when you are recording with a level sequence.

During testing, I found a crash on exit while in Multi-user that was missing a nullptr check in STakeRecorderCockpit.cpp

diff --git a/Engine/Plugins/Developer/Concert/ConcertSync/ConcertSyncClient/Source/ConcertSyncClient/Private/ConcertClientSequencerManager.cpp b/Engine/Plugins/Developer/Concert/ConcertSync/ConcertSyncClient/Source/ConcertSyncClient/Private/ConcertClientSequencerManager.cpp
--- a/Engine/Plugins/Developer/Concert/ConcertSync/ConcertSyncClient/Source/ConcertSyncClient/Private/ConcertClientSequencerManager.cpp
+++ b/Engine/Plugins/Developer/Concert/ConcertSync/ConcertSyncClient/Source/ConcertSyncClient/Private/ConcertClientSequencerManager.cpp
@@ -97,7 +97,7 @@ public:
 			return;
 		}
 
-		bool bShouldForceSend = bIsFirstEvent || CVarSequencerStatePacingEnabled.GetValueOnAnyThread() == 0;
+		bool bShouldForceSend = bIsFirstEvent || !IsPacingEnabled();
 		if (!bShouldForceSend && InEvent.State.bLoopMode && CurrentStateEvent)
 		{
 			// If we have loop enabled and we have started back over again then we should force a send.
@@ -130,7 +130,7 @@ public:
 			const double DeltaTime = CurrentTime - LastSendTime;
 			if (bForceSend || DeltaTime > CVarSequencerStatePacingDuration.GetValueOnGameThread())
 			{
-				bool bShouldSendReliably = bForceSend && CVarSequencerStatePacingEnabled.GetValueOnAnyThread();
+				bool bShouldSendReliably = bForceSend && IsPacingEnabled();
 				EConcertMessageFlags Flags = bShouldSendReliably || CVarSendStateEventsAsUnreliable.GetValueOnAnyThread() == 0 ? EConcertMessageFlags::ReliableOrdered : EConcertMessageFlags::None;
 				if (bForceSend || !bIsAContinuationOfPlay)
 				{
@@ -144,13 +144,26 @@ public:
 		}
 	}
 
+	bool IsPacingEnabled() const
+	{
+		return !bIsPacingSuspended && CVarSequencerStatePacingEnabled.GetValueOnAnyThread();
+	}
+
+	static void SetPacingSuspended(bool bInPacingSuspended)
+	{
+		bIsPacingSuspended = bInPacingSuspended;
+	}
+
 private:
 	TOptional<FConcertSequencerStateEvent> CurrentStateEvent;
 	double LastSendTime = 0;
 	bool bIsFirstEvent = true;
 	bool bIsAContinuationOfPlay = false;
+	static bool bIsPacingSuspended;
 };
 
+bool FConcertClientSequencerStateEventPacer::bIsPacingSuspended = false;
+
 class FConcertClientSequencePreloader : public TSharedFromThis<FConcertClientSequencePreloader>
 {
 public:
@@ -418,6 +431,12 @@ void SetConsoleVariableRespectingPriority(IConsoleVariable* AsVariable, bool bVa
 	AsVariable->Set( bValue ? 1 : 0, Flags );
 }
 
+void SetConsoleVariableRespectingPriority(IConsoleVariable* AsVariable, float InValue)
+{
+	EConsoleVariableFlags Flags = (EConsoleVariableFlags)( AsVariable->GetFlags() & ECVF_SetByMask );
+	AsVariable->Set( InValue, Flags );
+}
+
 bool FConcertClientSequencerManager::IsSequencerPlaybackSyncEnabled() const
 {
 	return CVarEnablePlaybackSync.GetValueOnAnyThread() > 0;
@@ -448,6 +467,16 @@ bool FConcertClientSequencerManager::IsSequencerRemoteCloseEnabled() const
 	return CVarEnableRemoteSequencerClose.GetValueOnAnyThread() > 0;
 }
 
+bool FConcertClientSequencerManager::IsSequencerPacingEnabled() const
+{
+	return CVarSequencerStatePacingEnabled.GetValueOnAnyThread();
+}
+
+float FConcertClientSequencerManager::SequencerPacingDuration() const
+{
+	return CVarSequencerStatePacingDuration.GetValueOnAnyThread();
+}
+
 void FConcertClientSequencerManager::SetSequencerRemoteOpen(bool bEnable)
 {
 	SetConsoleVariableRespectingPriority(CVarEnableRemoteSequencerOpen->AsVariable(), bEnable);
@@ -458,6 +487,26 @@ void FConcertClientSequencerManager::SetSequencerRemoteClose(bool bEnable)
 	SetConsoleVariableRespectingPriority(CVarEnableRemoteSequencerClose->AsVariable(), bEnable);
 }
 
+void FConcertClientSequencerManager::SetSequencerPacingEnabled(bool bEnable)
+{
+	SetConsoleVariableRespectingPriority(CVarSequencerStatePacingEnabled->AsVariable(), bEnable);
+}
+
+void FConcertClientSequencerManager::SetSequencerPacingDuration(float Duration)
+{
+	SetConsoleVariableRespectingPriority(CVarSequencerStatePacingDuration->AsVariable(), Duration);
+}
+
+void FConcertClientSequencerManager::SuspendSequencerPacing()
+{
+	FConcertClientSequencerStateEventPacer::SetPacingSuspended(true);
+}
+
+void FConcertClientSequencerManager::ResumeSequencerPacing()
+{
+	FConcertClientSequencerStateEventPacer::SetPacingSuspended(false);
+}
+
 bool FConcertClientSequencerManager::ShouldAlwaysCloseGameSequencerPlayer() const
 {
 	return CVarAlwaysCloseGamePlayerOnCloseEvent.GetValueOnAnyThread() > 0;
@@ -1428,6 +1477,15 @@ void FConcertClientSequencerManager::OnWorkspaceEndFrameCompleted()
 		return;
 	}
 
+	// Checking for the time adjustment event must happen first because any sequencer state changes may adjust the MovieScene range based on the state message.
+	// If this happens then we have no way of knowing if we need to adjust the take sections based on a new start position.   By making this check
+	// happen first we avoid this scenario and can properly check time adjustment events and move sections to the new start position.
+	for (const FConcertSequencerTimeAdjustmentEvent& Event : PendingTimeAdjustmentEvents)
+	{
+		ApplyTimeAdjustmentEvent(Event);
+	}
+	PendingTimeAdjustmentEvents.Reset();
+
 	for (const FString& Player: PendingCreate)
 	{
 		CreateNewSequencePlayerIfNotExists(Player);
@@ -1446,12 +1504,6 @@ void FConcertClientSequencerManager::OnWorkspaceEndFrameCompleted()
 	}
 	PendingSequencerEvents.Reset();
 
-	for (const FConcertSequencerTimeAdjustmentEvent& Event : PendingTimeAdjustmentEvents)
-	{
-		ApplyTimeAdjustmentEvent(Event);
-	}
-	PendingTimeAdjustmentEvents.Reset();
-
 	if (GEditor)
 	{
 		if (TSharedPtr<IConcertClientSession> Session = WeakSession.Pin())
diff --git a/Engine/Plugins/Developer/Concert/ConcertSync/ConcertSyncClient/Source/ConcertSyncClient/Private/ConcertClientSequencerManager.h b/Engine/Plugins/Developer/Concert/ConcertSync/ConcertSyncClient/Source/ConcertSyncClient/Private/ConcertClientSequencerManager.h
--- a/Engine/Plugins/Developer/Concert/ConcertSync/ConcertSyncClient/Source/ConcertSyncClient/Private/ConcertClientSequencerManager.h
+++ b/Engine/Plugins/Developer/Concert/ConcertSync/ConcertSyncClient/Source/ConcertSyncClient/Private/ConcertClientSequencerManager.h
@@ -90,6 +90,16 @@ public:
 	 */
 	virtual bool IsSequencerRemoteCloseEnabled() const override;
 
+	/**
+	 * @return true if the sequencer pacer is enabled.
+	 */
+	virtual bool IsSequencerPacingEnabled() const override;
+
+	/**
+	 * @return the current duration for the sequencer pacer.
+	 */
+	virtual float SequencerPacingDuration() const override;
+
 	/**
 	 * Checks the CVar to see if we are allowed to forcefully close the player on game instances.
 	 *
@@ -113,6 +123,30 @@ public:
 	 */
 	virtual void SetSequencerRemoteClose(bool bEnable) override;
 
+	/**
+	 * Specifies the sequencer pacing behavior while in a session.  This will set Concert.SequencerStatePacingEnabled console variable.
+	 *
+	 * @param bValue The value to be set for the sequencer pacing option
+	 */
+	virtual void SetSequencerPacingEnabled(bool bEnable) override;
+
+	/**
+	 * Sets the sequencer pacing duration if pacing is enabled. This will set Concert.SequencerStatePacingDuration console variable.
+	 *
+	 * @param Duration The spacing between state event messages.  The default value is 1.0
+	 */
+	virtual void SetSequencerPacingDuration(float Duration) override;
+
+	/**
+	 * Temporarily suspend sequencer pacing on this client.
+	 */
+	virtual void SuspendSequencerPacing() override;
+
+	/**
+	 * Resume sequencer pacing on this client if the console variable is enabled.
+	 */
+	virtual void ResumeSequencerPacing() override;
+
 	/** Assign the current active workspace to this sequencer. */
 	void SetActiveWorkspace(TSharedPtr<FConcertClientWorkspace> Workspace);
 
diff --git a/Engine/Plugins/Developer/Concert/ConcertSync/ConcertSyncClient/Source/ConcertSyncClient/Public/IConcertClientSequencerManager.h b/Engine/Plugins/Developer/Concert/ConcertSync/ConcertSyncClient/Source/ConcertSyncClient/Public/IConcertClientSequencerManager.h
--- a/Engine/Plugins/Developer/Concert/ConcertSync/ConcertSyncClient/Source/ConcertSyncClient/Public/IConcertClientSequencerManager.h
+++ b/Engine/Plugins/Developer/Concert/ConcertSync/ConcertSyncClient/Source/ConcertSyncClient/Public/IConcertClientSequencerManager.h
@@ -45,6 +45,16 @@ public:
 	 */
 	virtual bool IsSequencerRemoteCloseEnabled() const = 0;
 
+	/**
+	 * @return true if the sequencer pacer is enabled.
+	 */
+	virtual bool IsSequencerPacingEnabled() const = 0;
+
+	/**
+	 * @return the current duration for the sequencer pacer.
+	 */
+	virtual float SequencerPacingDuration() const = 0;
+
 	/**
 	 * Set the remote open option in Multi-User
 	 * which opens Sequencer for other users when this option is enabled on both user machines.
@@ -60,4 +70,28 @@ public:
 	 * @param bEnable The value to set for the remote close option
 	 */
 	virtual void SetSequencerRemoteClose(bool bEnable) = 0;
+
+	/**
+	 * Temporarily suspend sequencer pacing on this client.
+	 */
+	virtual void SuspendSequencerPacing() = 0;
+
+	/**
+	 * Resume sequencer pacing on this client if the console variable is enabled.
+	 */
+	virtual void ResumeSequencerPacing() = 0;
+
+	/**
+	 * Specifies the sequencer pacing behavior while in a session.
+	 *
+	 * @param bValue The value to be set for the sequencer pacing option
+	 */
+	virtual void SetSequencerPacingEnabled(bool bEnable) = 0;
+
+	/**
+	 * Sets the sequencer pacing duration value.
+	 *
+	 * @param Duration The spacing between state event messages.
+	 */
+	virtual void SetSequencerPacingDuration(float Duration) = 0;
 };
diff --git a/Engine/Plugins/VirtualProduction/MultiUserTakes/Source/ConcertTakeRecorder/Private/ConcertTakeRecorderManager.cpp b/Engine/Plugins/VirtualProduction/MultiUserTakes/Source/ConcertTakeRecorder/Private/ConcertTakeRecorderManager.cpp
--- a/Engine/Plugins/VirtualProduction/MultiUserTakes/Source/ConcertTakeRecorder/Private/ConcertTakeRecorderManager.cpp
+++ b/Engine/Plugins/VirtualProduction/MultiUserTakes/Source/ConcertTakeRecorder/Private/ConcertTakeRecorderManager.cpp
@@ -9,6 +9,7 @@
 #include "Engine/Engine.h"
 #include "IConcertClientTransactionBridge.h"
 #include "IConcertClientPackageBridge.h"
+#include "IConcertClientSequencerManager.h"
 #include "IConcertClient.h"
 #include "IConcertSyncClient.h"
 #include "IConcertSyncClientModule.h"
@@ -454,6 +455,13 @@ void FConcertTakeRecorderManager::OnTakeRecorderInitialized(UTakeRecorder* TakeR
 				{
 					TakeRecorder->SetDisableSaveTick(true);
 				}
+				else
+				{
+					if (TSharedPtr<IConcertSyncClient> ConcertSyncClient = IConcertSyncClientModule::Get().GetClient(TEXT("MultiUser")))
+					{
+						ConcertSyncClient->GetSequencerManager()->SuspendSequencerPacing();
+					}
+				}
 				FConcertTakeInitializedEvent TakeInitializedEvent;
 				TakeInitializedEvent.TakeName = TakeRecorder->GetName();
 				TakeInitializedEvent.TakePresetPath = TakeMetaData->GetPresetOrigin()->GetPathName();
@@ -485,6 +493,11 @@ void FConcertTakeRecorderManager::OnRecordingFinished(UTakeRecorder* TakeRecorde
 
 			if (CanRecord())
 			{
+				if (TSharedPtr<IConcertSyncClient> ConcertSyncClient = IConcertSyncClientModule::Get().GetClient(TEXT("MultiUser")))
+				{
+					ConcertSyncClient->GetSequencerManager()->ResumeSequencerPacing();
+				}
+
 				LastLevelSequence = TakeRecorder->GetSequence();
 				check(LastLevelSequence);
 				FConcertRecordingNamedLevelSequenceEvent NamedSequence{LastLevelSequence->GetPathName()};
diff --git a/Engine/Plugins/VirtualProduction/Takes/Source/TakeRecorder/Private/Widgets/STakeRecorderCockpit.cpp b/Engine/Plugins/VirtualProduction/Takes/Source/TakeRecorder/Private/Widgets/STakeRecorderCockpit.cpp
--- a/Engine/Plugins/VirtualProduction/Takes/Source/TakeRecorder/Private/Widgets/STakeRecorderCockpit.cpp
+++ b/Engine/Plugins/VirtualProduction/Takes/Source/TakeRecorder/Private/Widgets/STakeRecorderCockpit.cpp
@@ -1017,7 +1017,7 @@ FReply STakeRecorderCockpit::OnAddMarkedFrame()
 
 bool STakeRecorderCockpit::Reviewing() const 
 {
-	return bool(!Recording() && (TakeMetaData->Recorded() && TakeRecorderModeAttribute.Get() != ETakeRecorderMode::RecordIntoSequence));
+	return bool(!Recording() && (TakeMetaData && TakeMetaData->Recorded() && TakeRecorderModeAttribute.Get() != ETakeRecorderMode::RecordIntoSequence));
 }
 
 bool STakeRecorderCockpit::Recording() const
