Add IConcertClientReplicationManager::HasAuthorityOver, FClientReplicationDataCollector::OwnsObjectInAnyStream, and FClientReplicationDataCollector::GetStreamsOwningObject.

diff --git a/Engine/Plugins/Developer/Concert/ConcertSync/ConcertSyncClient/Source/ConcertSyncClient/Private/Replication/Manager/ReplicationManager.cpp b/Engine/Plugins/Developer/Concert/ConcertSync/ConcertSyncClient/Source/ConcertSyncClient/Private/Replication/Manager/ReplicationManager.cpp
--- a/Engine/Plugins/Developer/Concert/ConcertSync/ConcertSyncClient/Source/ConcertSyncClient/Private/Replication/Manager/ReplicationManager.cpp
+++ b/Engine/Plugins/Developer/Concert/ConcertSync/ConcertSyncClient/Source/ConcertSyncClient/Private/Replication/Manager/ReplicationManager.cpp
@@ -104,6 +104,12 @@ namespace UE::ConcertSyncClient::Replication
 			: TSet<FGuid>{}; 
 	}
 
+	bool FReplicationManager::HasAuthorityOver(const FSoftObjectPath& ObjectPath) const
+	{
+		return ensureMsgf(CurrentState, TEXT("StartAcceptingJoinRequests should have been called at this point."))
+			&& CurrentState->HasAuthorityOver(ObjectPath);
+	}
+
 	IConcertClientReplicationManager::ESyncControlEnumerationResult FReplicationManager::ForEachSyncControlledObject(TFunctionRef<EBreakBehavior(const FConcertObjectInStreamID& Object)> Callback) const
 	{
 		// Check() to avoid returning some dummy static variable
diff --git a/Engine/Plugins/Developer/Concert/ConcertSync/ConcertSyncClient/Source/ConcertSyncClient/Private/Replication/Manager/ReplicationManager.h b/Engine/Plugins/Developer/Concert/ConcertSync/ConcertSyncClient/Source/ConcertSyncClient/Private/Replication/Manager/ReplicationManager.h
--- a/Engine/Plugins/Developer/Concert/ConcertSync/ConcertSyncClient/Source/ConcertSyncClient/Private/Replication/Manager/ReplicationManager.h
+++ b/Engine/Plugins/Developer/Concert/ConcertSync/ConcertSyncClient/Source/ConcertSyncClient/Private/Replication/Manager/ReplicationManager.h
@@ -40,6 +40,7 @@ namespace UE::ConcertSyncClient::Replication
 		virtual TFuture<FConcertReplication_ChangeStream_Response> ChangeStream(FConcertReplication_ChangeStream_Request Args) override;
 		virtual EAuthorityEnumerationResult ForEachClientOwnedObject(TFunctionRef<EBreakBehavior(const FSoftObjectPath& Object, TSet<FGuid>&& OwningStreams)>) const override;
 		virtual TSet<FGuid> GetClientOwnedStreamsForObject(const FSoftObjectPath& ObjectPath) const override;
+		virtual bool HasAuthorityOver(const FSoftObjectPath& ObjectPath) const override;
 		virtual ESyncControlEnumerationResult ForEachSyncControlledObject(TFunctionRef<EBreakBehavior(const FConcertObjectInStreamID& Object)> Callback) const override;
 		virtual uint32 NumSyncControlledObjects() const override;
 		virtual bool HasSyncControl(const FConcertObjectInStreamID& Object) const override;
diff --git a/Engine/Plugins/Developer/Concert/ConcertSync/ConcertSyncClient/Source/ConcertSyncClient/Private/Replication/Manager/ReplicationManagerState.h b/Engine/Plugins/Developer/Concert/ConcertSync/ConcertSyncClient/Source/ConcertSyncClient/Private/Replication/Manager/ReplicationManagerState.h
--- a/Engine/Plugins/Developer/Concert/ConcertSync/ConcertSyncClient/Source/ConcertSyncClient/Private/Replication/Manager/ReplicationManagerState.h
+++ b/Engine/Plugins/Developer/Concert/ConcertSync/ConcertSyncClient/Source/ConcertSyncClient/Private/Replication/Manager/ReplicationManagerState.h
@@ -29,6 +29,7 @@ namespace UE::ConcertSyncClient::Replication
 		virtual TFuture<FConcertReplication_ChangeStream_Response> ChangeStream(FConcertReplication_ChangeStream_Request Args) override { return MakeFulfilledPromise<FConcertReplication_ChangeStream_Response>().GetFuture(); }
 		virtual EAuthorityEnumerationResult ForEachClientOwnedObject(TFunctionRef<EBreakBehavior(const FSoftObjectPath& Object, TSet<FGuid>&& OwningStreams)> Callback) const override { return EAuthorityEnumerationResult::NoAuthorityAvailable; }
 		virtual TSet<FGuid> GetClientOwnedStreamsForObject(const FSoftObjectPath& ObjectPath) const override { return {}; }
+		virtual bool HasAuthorityOver(const FSoftObjectPath& ObjectPath) const override { return false; }
 		virtual ESyncControlEnumerationResult ForEachSyncControlledObject(TFunctionRef<EBreakBehavior(const FConcertObjectInStreamID& Object)> Callback) const override { return ESyncControlEnumerationResult::NoneAvailable; }
 		virtual uint32 NumSyncControlledObjects() const override { return 0; }
 		virtual bool HasSyncControl(const FConcertObjectInStreamID& Object) const override { return false; }
diff --git a/Engine/Plugins/Developer/Concert/ConcertSync/ConcertSyncClient/Source/ConcertSyncClient/Private/Replication/Manager/ReplicationManagerState_Connected.cpp b/Engine/Plugins/Developer/Concert/ConcertSync/ConcertSyncClient/Source/ConcertSyncClient/Private/Replication/Manager/ReplicationManagerState_Connected.cpp
--- a/Engine/Plugins/Developer/Concert/ConcertSync/ConcertSyncClient/Source/ConcertSyncClient/Private/Replication/Manager/ReplicationManagerState_Connected.cpp
+++ b/Engine/Plugins/Developer/Concert/ConcertSync/ConcertSyncClient/Source/ConcertSyncClient/Private/Replication/Manager/ReplicationManagerState_Connected.cpp
@@ -217,6 +217,11 @@ namespace UE::ConcertSyncClient::Replication
 		return Result;
 	}
 
+	bool FReplicationManagerState_Connected::HasAuthorityOver(const FSoftObjectPath& ObjectPath) const
+	{
+		return ReplicationDataSource.OwnsObjectInAnyStream(ObjectPath);
+	}
+
 	IConcertClientReplicationManager::ESyncControlEnumerationResult FReplicationManagerState_Connected::ForEachSyncControlledObject(TFunctionRef<EBreakBehavior(const FConcertObjectInStreamID& Object)> Callback) const
 	{
 		return SyncControl.EnumerateAllowedObjects(Callback)
diff --git a/Engine/Plugins/Developer/Concert/ConcertSync/ConcertSyncClient/Source/ConcertSyncClient/Private/Replication/Manager/ReplicationManagerState_Connected.h b/Engine/Plugins/Developer/Concert/ConcertSync/ConcertSyncClient/Source/ConcertSyncClient/Private/Replication/Manager/ReplicationManagerState_Connected.h
--- a/Engine/Plugins/Developer/Concert/ConcertSync/ConcertSyncClient/Source/ConcertSyncClient/Private/Replication/Manager/ReplicationManagerState_Connected.h
+++ b/Engine/Plugins/Developer/Concert/ConcertSync/ConcertSyncClient/Source/ConcertSyncClient/Private/Replication/Manager/ReplicationManagerState_Connected.h
@@ -52,6 +52,7 @@ namespace UE::ConcertSyncClient::Replication
 		virtual TFuture<FConcertReplication_ChangeStream_Response> ChangeStream(FConcertReplication_ChangeStream_Request Args) override;
 		virtual EAuthorityEnumerationResult ForEachClientOwnedObject(TFunctionRef<EBreakBehavior(const FSoftObjectPath& Object, TSet<FGuid>&& OwningStreams)> Callback) const override;
 		virtual TSet<FGuid> GetClientOwnedStreamsForObject(const FSoftObjectPath& ObjectPath) const override;
+		virtual bool HasAuthorityOver(const FSoftObjectPath& ObjectPath) const override;
 		virtual ESyncControlEnumerationResult ForEachSyncControlledObject(TFunctionRef<EBreakBehavior(const FConcertObjectInStreamID& Object)> Callback) const override;
 		virtual uint32 NumSyncControlledObjects() const override { return SyncControl.Num(); }
 		virtual bool HasSyncControl(const FConcertObjectInStreamID& Object) const override { return SyncControl.IsObjectAllowed(Object); }
diff --git a/Engine/Plugins/Developer/Concert/ConcertSync/ConcertSyncClient/Source/ConcertSyncClient/Private/Replication/Processing/ClientReplicationDataCollector.cpp b/Engine/Plugins/Developer/Concert/ConcertSync/ConcertSyncClient/Source/ConcertSyncClient/Private/Replication/Processing/ClientReplicationDataCollector.cpp
--- a/Engine/Plugins/Developer/Concert/ConcertSync/ConcertSyncClient/Source/ConcertSyncClient/Private/Replication/Processing/ClientReplicationDataCollector.cpp
+++ b/Engine/Plugins/Developer/Concert/ConcertSync/ConcertSyncClient/Source/ConcertSyncClient/Private/Replication/Processing/ClientReplicationDataCollector.cpp
@@ -160,6 +160,24 @@ namespace UE::ConcertSyncClient::Replication
 		return ObjectInfo && ObjectInfo->ContainsByPredicate([&StreamId](const FObjectInfo& ObjectInfo){ return ObjectInfo.StreamId == StreamId; });
 	}
 
+	bool FClientReplicationDataCollector::OwnsObjectInAnyStream(const FSoftObjectPath& ObjectPath) const
+	{
+		return ObjectsToReplicate.Contains(ObjectPath);
+	}
+
+	TArray<FGuid> FClientReplicationDataCollector::GetStreamsOwningObject(const FSoftObjectPath& ObjectPath) const
+	{
+		const TArray<FObjectInfo>* ObjectInfo = ObjectsToReplicate.Find(ObjectPath);
+		if (!ObjectInfo)
+		{
+			return {};
+		}
+
+		TArray<FGuid> Result;
+		Algo::Transform(*ObjectInfo, Result, [](const FObjectInfo& ObjectInfo){ return ObjectInfo.StreamId; });
+		return Result;
+	}
+
 	void FClientReplicationDataCollector::ForEachPendingObject(TFunctionRef<void(const ConcertSyncCore::FPendingObjectReplicationInfo&)> ProcessItemFunc) const
 	{
 		for (const TPair<FSoftObjectPath, TArray<FObjectInfo>>& Pair : ObjectsToReplicate)
diff --git a/Engine/Plugins/Developer/Concert/ConcertSync/ConcertSyncClient/Source/ConcertSyncClient/Private/Replication/Processing/ClientReplicationDataCollector.h b/Engine/Plugins/Developer/Concert/ConcertSync/ConcertSyncClient/Source/ConcertSyncClient/Private/Replication/Processing/ClientReplicationDataCollector.h
--- a/Engine/Plugins/Developer/Concert/ConcertSync/ConcertSyncClient/Source/ConcertSyncClient/Private/Replication/Processing/ClientReplicationDataCollector.h
+++ b/Engine/Plugins/Developer/Concert/ConcertSync/ConcertSyncClient/Source/ConcertSyncClient/Private/Replication/Processing/ClientReplicationDataCollector.h
@@ -80,6 +80,11 @@ namespace UE::ConcertSyncClient::Replication
 
 		/** @return Whether ObjectPath is owned with StreamId. */
 		bool OwnsObjectInStream(const FSoftObjectPath& ObjectPath, const FGuid& StreamId) const;
+		/** @return Whether ObjectPath owns ObjectPath in any stream. */
+		bool OwnsObjectInAnyStream(const FSoftObjectPath& ObjectPath) const;
+
+		/** @return All streams containing ObjectPath that we have authority over within that stream. */
+		TArray<FGuid> GetStreamsOwningObject(const FSoftObjectPath& ObjectPath) const;
 
 		//~ Begin IReplicationDataSource Interface
 		virtual void ForEachPendingObject(TFunctionRef<void(const ConcertSyncCore::FPendingObjectReplicationInfo&)> ProcessItemFunc) const override;
diff --git a/Engine/Plugins/Developer/Concert/ConcertSync/ConcertSyncClient/Source/ConcertSyncClient/Public/Replication/IConcertClientReplicationManager.h b/Engine/Plugins/Developer/Concert/ConcertSync/ConcertSyncClient/Source/ConcertSyncClient/Public/Replication/IConcertClientReplicationManager.h
--- a/Engine/Plugins/Developer/Concert/ConcertSync/ConcertSyncClient/Source/ConcertSyncClient/Public/Replication/IConcertClientReplicationManager.h
+++ b/Engine/Plugins/Developer/Concert/ConcertSync/ConcertSyncClient/Source/ConcertSyncClient/Public/Replication/IConcertClientReplicationManager.h
@@ -104,6 +104,8 @@ public:
 	virtual EAuthorityEnumerationResult ForEachClientOwnedObject(TFunctionRef<EBreakBehavior(const FSoftObjectPath& Object, TSet<FGuid>&& OwningStreams)> Callback) const = 0;
 	/** @return All the streams this client has registered that have authority for the given ObjectPath. */
 	virtual TSet<FGuid> GetClientOwnedStreamsForObject(const FSoftObjectPath& ObjectPath) const = 0;
+	/** @return Whether this client has authority over ObjectPath for some stream. */
+	virtual bool HasAuthorityOver(const FSoftObjectPath& ObjectPath) const = 0;
 	/** @return All owned objects and associated owning streams. */
 	TMap<FSoftObjectPath, TSet<FGuid>> GetClientOwnedObjects() const;
 
