Fix: If a client had no authority, made a request to release authority, and that request failed (e.g. timeout), the revert mechanism would incorrectly make the local client think that is has authority over the released objects.

diff --git a/Engine/Plugins/Developer/Concert/ConcertSync/ConcertSyncClient/Source/ConcertSyncClient/Private/Replication/Manager/ReplicationManagerState_Connected.cpp b/Engine/Plugins/Developer/Concert/ConcertSync/ConcertSyncClient/Source/ConcertSyncClient/Private/Replication/Manager/ReplicationManagerState_Connected.cpp
--- a/Engine/Plugins/Developer/Concert/ConcertSync/ConcertSyncClient/Source/ConcertSyncClient/Private/Replication/Manager/ReplicationManagerState_Connected.cpp
+++ b/Engine/Plugins/Developer/Concert/ConcertSync/ConcertSyncClient/Source/ConcertSyncClient/Private/Replication/Manager/ReplicationManagerState_Connected.cpp
@@ -23,6 +23,59 @@
 #include "UObject/UObjectGlobals.h"
 #include "Utils/ContentComparisionUtils.h"
 
+namespace UE::ConcertSyncClient::Replication::Private
+{
+	/** Removes ObjectsToRemove and returns the authority that was actually removed, possibly to be reverted later. */
+	template<typename TStreamArrayType, typename TProjectStreamArray> requires std::is_invocable_r_v<const TArray<FGuid>&, TProjectStreamArray, const TStreamArrayType&>
+	static TMap<FSoftObjectPath, TArray<FGuid>> RemoveObjectsAndTrackRemovedAuthority(
+		FClientReplicationDataCollector& ReplicationDataSource,
+		const TMap<FSoftObjectPath, TStreamArrayType>& ObjectsToRemove,
+		TProjectStreamArray&& Project
+		)
+	{
+		TMap<FSoftObjectPath, TArray<FGuid>> RevertedContent;
+		for (const TPair<FSoftObjectPath, TStreamArrayType>& ReleaseAuthority : ObjectsToRemove)
+		{
+			const FSoftObjectPath& ObjectPath = ReleaseAuthority.Key;
+			const TArray<FGuid>& StreamsToRemove = Project(ReleaseAuthority.Value);
+
+			// Remember which streams we actually removed, so we can revert it again.
+			TArray<FGuid> StreamsActuallyRemoved = ReplicationDataSource.GetStreamsOwningObject(ObjectPath);
+			StreamsActuallyRemoved.SetNum(Algo::RemoveIf(StreamsActuallyRemoved, [&StreamsToRemove](const FGuid& StreamId)
+			{
+				return !StreamsToRemove.Contains(StreamId);
+			}));
+			RevertedContent.Add(ObjectPath, StreamsActuallyRemoved);
+			
+			ReplicationDataSource.RemoveReplicatedObjectStreams(ObjectPath, StreamsToRemove);
+		}
+		return RevertedContent;
+	}
+
+	static TMap<FSoftObjectPath, TArray<FGuid>> RemoveObjectsAndTrackRemovedAuthority(
+		FClientReplicationDataCollector& ReplicationDataSource,
+		const TMap<FSoftObjectPath, FConcertStreamArray>& ObjectsToRemove
+		)
+	{
+		return RemoveObjectsAndTrackRemovedAuthority<FConcertStreamArray>(
+			ReplicationDataSource, ObjectsToRemove,
+			[](const FConcertStreamArray& StreamArray)-> const TArray<FGuid>& { return StreamArray.StreamIds; }
+			); 
+	}
+	static TMap<FSoftObjectPath, TArray<FGuid>> RemoveObjectsAndTrackRemovedAuthority(
+		FClientReplicationDataCollector& ReplicationDataSource,
+		const TMap<FSoftObjectPath, TArray<FGuid>>& ObjectsToRemove
+		)
+	{
+		return RemoveObjectsAndTrackRemovedAuthority<TArray<FGuid>>(
+			ReplicationDataSource, ObjectsToRemove,
+			[](const TArray<FGuid>& StreamArray)-> const TArray<FGuid>& { return StreamArray; }
+			); 
+	}
+	
+	
+}
+
 namespace UE::ConcertSyncClient::Replication
 {
 	FReplicationManagerState_Connected::FReplicationManagerState_Connected(
@@ -111,29 +164,31 @@ namespace UE::ConcertSyncClient::Replication
 		
 		// Stop replicating removed objects right now: the server will remove authority after processing this request.
 		// At that point, it will log errors for receiving replication data from a client without authority.
-		ApplyAuthorityChangeRemovedObjects(Args);
+		TMap<FSoftObjectPath, TArray<FGuid>> PredictedChange = ApplyAuthorityChangeRemovedObjects(Args);
 		// We don't need worry about updating sync control until it is processed below - the local client will not attempt to replicate the object
 		// because we just locally updated the authority cache.
 
 		LogNetworkMessage(CVarLogAuthorityRequestsAndResponsesOnClient, Args);
 		return LiveSession->SendCustomRequest<FConcertReplication_ChangeAuthority_Request, FConcertReplication_ChangeAuthority_Response>(Args, LiveSession->GetSessionServerEndpointId())
-			.Next([WeakThis = TWeakPtr<FReplicationManagerState_Connected>(SharedThis(this)), Args](FConcertReplication_ChangeAuthority_Response&& Response) mutable
-			{
-				LogNetworkMessage(CVarLogAuthorityRequestsAndResponsesOnClient, Response);
-				
-				if (const TSharedPtr<FReplicationManagerState_Connected> ThisPin = WeakThis.Pin()
-					; ThisPin && Response.ErrorCode == EReplicationResponseErrorCode::Handled)
-				{
-					ThisPin->FinalizePredictedAuthorityChange(Args, Response.RejectedObjects, Response.SyncControl);
-				}
-				else if (ThisPin && Response.ErrorCode == EReplicationResponseErrorCode::Timeout)
+			.Next(
+				[WeakThis = TWeakPtr<FReplicationManagerState_Connected>(SharedThis(this)), Args, PredictedChange = MoveTemp(PredictedChange)]
+				(FConcertReplication_ChangeAuthority_Response&& Response) mutable
 				{
-					// HandleReleasingReplicatedObjects caused Args.ReleaseAuthority to stop being replicated. Revert.
-					ThisPin->RevertAuthorityChangeReleasedObjects(Args);
-				}
-
-				return FConcertReplication_ChangeAuthority_Response { MoveTemp(Response) };
-			});
+					LogNetworkMessage(CVarLogAuthorityRequestsAndResponsesOnClient, Response);
+					
+					if (const TSharedPtr<FReplicationManagerState_Connected> ThisPin = WeakThis.Pin()
+						; ThisPin && Response.ErrorCode == EReplicationResponseErrorCode::Handled)
+					{
+						ThisPin->FinalizePredictedAuthorityChange(Args, Response.RejectedObjects, Response.SyncControl);
+					}
+					else if (ThisPin && Response.ErrorCode == EReplicationResponseErrorCode::Timeout)
+					{
+						// ApplyAuthorityChangeRemovedObjects caused PredictedChange to stop being replicated. Revert.
+						ThisPin->RevertAuthorityChangeReleasedObjects(PredictedChange);
+					}
+
+					return FConcertReplication_ChangeAuthority_Response { MoveTemp(Response) };
+				});
 	}
 
 	TFuture<FConcertReplication_QueryReplicationInfo_Response> FReplicationManagerState_Connected::QueryClientInfo(FConcertReplication_QueryReplicationInfo_Request Args)
@@ -165,7 +220,7 @@ namespace UE::ConcertSyncClient::Replication
 		
 		// Stop replicating removed objects right now: the server will remove authority after processing this request.
 		// At that point, it will log errors for receiving replication data from a client without authority.
-		TMap<FSoftObjectPath, TArray<FGuid>> PredictedChange = PredictAndApplyStreamChangeRemovedObjects(Args);
+		FChangeStreamPredictedChange PredictedChange = PredictAndApplyStreamChangeRemovedObjects(Args);
 		// We don't need worry about updating sync control until it is processed below - the local client will not attempt to replicate the object
 		// because we just locally updated the replication cache.
 		
@@ -323,24 +378,24 @@ namespace UE::ConcertSyncClient::Replication
 		ReplicationApplier.ProcessObjects(Params);
 	}
 
-	TMap<FSoftObjectPath, TArray<FGuid>> FReplicationManagerState_Connected::PredictAndApplyStreamChangeRemovedObjects(const FConcertReplication_ChangeStream_Request& Request)
+	FChangeStreamPredictedChange FReplicationManagerState_Connected::PredictAndApplyStreamChangeRemovedObjects(const FConcertReplication_ChangeStream_Request& Request)
 	{
-		TMap<FSoftObjectPath, TArray<FGuid>> PredicatedRemovedObjects = ComputeRemovedObjects(RegisteredStreams, Request);
-		ApplyStreamChangeRemovedObjects(PredicatedRemovedObjects);
-		return PredicatedRemovedObjects;
+		FChangeStreamPredictedChange Change;
+		Change.ObjectsRemovedFromStream = ComputeRemovedObjects(RegisteredStreams, Request);
+		Change.AuthorityRemovedFromStreams = ApplyStreamChangeRemovedObjects(Change.ObjectsRemovedFromStream);
+		return Change;
 	}
 
-	void FReplicationManagerState_Connected::ApplyStreamChangeRemovedObjects(const TMap<FSoftObjectPath, TArray<FGuid>>& PredicatedRemovedObjects)
+	TMap<FSoftObjectPath, TArray<FGuid>> FReplicationManagerState_Connected::ApplyStreamChangeRemovedObjects(
+		const TMap<FSoftObjectPath, TArray<FGuid>>& PredicatedRemovedObjects
+		)
 	{
-		for (const TPair<FSoftObjectPath, TArray<FGuid>>& RemovedObjectInfo : PredicatedRemovedObjects)
-		{
-			ReplicationDataSource.RemoveReplicatedObjectStreams(RemovedObjectInfo.Key, RemovedObjectInfo.Value);
-		}
+		return Private::RemoveObjectsAndTrackRemovedAuthority(ReplicationDataSource, PredicatedRemovedObjects);
 	}
 
-	void FReplicationManagerState_Connected::RevertPredictedStreamChangeRemovedObjects(const TMap<FSoftObjectPath, TArray<FGuid>>& PredictedChange)
+	void FReplicationManagerState_Connected::RevertPredictedStreamChangeRemovedObjects(const FChangeStreamPredictedChange& PredictedChange)
 	{
-		for (const TPair<FSoftObjectPath, TArray<FGuid>>& RemovedObjectInfo : PredictedChange)
+		for (const TPair<FSoftObjectPath, TArray<FGuid>>& RemovedObjectInfo : PredictedChange.AuthorityRemovedFromStreams)
 		{
 			ReplicationDataSource.AddReplicatedObjectStreams(RemovedObjectInfo.Key, RemovedObjectInfo.Value);
 		}
@@ -387,19 +442,16 @@ namespace UE::ConcertSyncClient::Replication
 		}
 	}
 
-	void FReplicationManagerState_Connected::ApplyAuthorityChangeRemovedObjects(const FConcertReplication_ChangeAuthority_Request& Request)
+	TMap<FSoftObjectPath, TArray<FGuid>> FReplicationManagerState_Connected::ApplyAuthorityChangeRemovedObjects(const FConcertReplication_ChangeAuthority_Request& Request)
 	{
-		for (const TPair<FSoftObjectPath, FConcertStreamArray>& ReleaseAuthority : Request.ReleaseAuthority)
-		{
-			ReplicationDataSource.RemoveReplicatedObjectStreams(ReleaseAuthority.Key, ReleaseAuthority.Value.StreamIds);
-		}
+		return Private::RemoveObjectsAndTrackRemovedAuthority(ReplicationDataSource, Request.ReleaseAuthority);
 	}
 
-	void FReplicationManagerState_Connected::RevertAuthorityChangeReleasedObjects(const FConcertReplication_ChangeAuthority_Request& Request)
+	void FReplicationManagerState_Connected::RevertAuthorityChangeReleasedObjects(const TMap<FSoftObjectPath, TArray<FGuid>>& PredictedChange)
 	{
-		for (const TPair<FSoftObjectPath, FConcertStreamArray>& ReleaseAuthority : Request.ReleaseAuthority)
+		for (const TPair<FSoftObjectPath, TArray<FGuid>>& ReleaseAuthority : PredictedChange)
 		{
-			ReplicationDataSource.AddReplicatedObjectStreams(ReleaseAuthority.Key, ReleaseAuthority.Value.StreamIds);
+			ReplicationDataSource.AddReplicatedObjectStreams(ReleaseAuthority.Key, ReleaseAuthority.Value);
 		}
 	}
 
diff --git a/Engine/Plugins/Developer/Concert/ConcertSync/ConcertSyncClient/Source/ConcertSyncClient/Private/Replication/Manager/ReplicationManagerState_Connected.h b/Engine/Plugins/Developer/Concert/ConcertSync/ConcertSyncClient/Source/ConcertSyncClient/Private/Replication/Manager/ReplicationManagerState_Connected.h
--- a/Engine/Plugins/Developer/Concert/ConcertSync/ConcertSyncClient/Source/ConcertSyncClient/Private/Replication/Manager/ReplicationManagerState_Connected.h
+++ b/Engine/Plugins/Developer/Concert/ConcertSync/ConcertSyncClient/Source/ConcertSyncClient/Private/Replication/Manager/ReplicationManagerState_Connected.h
@@ -20,6 +20,12 @@ namespace UE::ConcertSyncCore { class FObjectReplicationReceiver; }
 
 namespace UE::ConcertSyncClient::Replication
 {
+	struct FChangeStreamPredictedChange
+	{
+		TMap<FSoftObjectPath, TArray<FGuid>> ObjectsRemovedFromStream;
+		TMap<FSoftObjectPath, TArray<FGuid>> AuthorityRemovedFromStreams;
+	};
+	
 	/**
 	 * State for when the client has successfully completed a replication handshake.
 	 *
@@ -117,19 +123,19 @@ namespace UE::ConcertSyncClient::Replication
 		void HandleChangeClientEvent(const FConcertSessionContext& Context, const FConcertReplication_ChangeClientEvent& Event);
 
 		/** Changes the local state assuming that Request will succeed. */
-		TMap<FSoftObjectPath, TArray<FGuid>> PredictAndApplyStreamChangeRemovedObjects(const FConcertReplication_ChangeStream_Request& Request);
-		void ApplyStreamChangeRemovedObjects(const TMap<FSoftObjectPath, TArray<FGuid>>& PredicatedRemovedObjects);
+		FChangeStreamPredictedChange PredictAndApplyStreamChangeRemovedObjects(const FConcertReplication_ChangeStream_Request& Request);
+		TMap<FSoftObjectPath, TArray<FGuid>> ApplyStreamChangeRemovedObjects(const TMap<FSoftObjectPath, TArray<FGuid>>& PredicatedRemovedObjects);
 		/** Reverts changes previously made by PredictStreamChangeRemovedObjects. */
-		void RevertPredictedStreamChangeRemovedObjects(const TMap<FSoftObjectPath, TArray<FGuid>>& PredictedChange);
+		void RevertPredictedStreamChangeRemovedObjects(const FChangeStreamPredictedChange& PredictedChange);
 		/** Applies stream changes that we previously predicted using PredictStreamChangeRemovedObjects. */
 		void FinalizePredictedStreamChange(const FConcertReplication_ChangeStream_Request& StreamChange);
 		/** Updates replicated objects affected by the change request. */
 		void UpdateReplicatedObjectsAfterStreamChange(const FConcertReplication_ChangeStream_Request& Request);
 		
 		/** Changes the local state assuming that Request will succeed. */
-		void ApplyAuthorityChangeRemovedObjects(const FConcertReplication_ChangeAuthority_Request& Request);
+		TMap<FSoftObjectPath, TArray<FGuid>> ApplyAuthorityChangeRemovedObjects(const FConcertReplication_ChangeAuthority_Request& Request);
 		/** Reverts changes previously made by PredictAuthorityChangeReleasedObjects. */
-		void RevertAuthorityChangeReleasedObjects(const FConcertReplication_ChangeAuthority_Request& Request);
+		void RevertAuthorityChangeReleasedObjects(const TMap<FSoftObjectPath, TArray<FGuid>>& PredictedChange);
 		/** Applies authority changes that we previously predicted using PredictAuthorityChangeReleasedObjects. */
 		void FinalizePredictedAuthorityChange(const FConcertReplication_ChangeAuthority_Request& AuthorityChange, const TMap<FSoftObjectPath, FConcertStreamArray>& RejectedObjects, const FConcertReplication_ChangeSyncControl& SyncControlChange);
 		/** Updates the objects which should be replicated after changing authority. */
diff --git a/Engine/Plugins/Developer/Concert/ConcertSync/ConcertSyncClient/Source/ConcertSyncClient/Private/Replication/Manager/ReplicationManagerState_Connected_PutState.cpp b/Engine/Plugins/Developer/Concert/ConcertSync/ConcertSyncClient/Source/ConcertSyncClient/Private/Replication/Manager/ReplicationManagerState_Connected_PutState.cpp
--- a/Engine/Plugins/Developer/Concert/ConcertSync/ConcertSyncClient/Source/ConcertSyncClient/Private/Replication/Manager/ReplicationManagerState_Connected_PutState.cpp
+++ b/Engine/Plugins/Developer/Concert/ConcertSync/ConcertSyncClient/Source/ConcertSyncClient/Private/Replication/Manager/ReplicationManagerState_Connected_PutState.cpp
@@ -117,7 +117,10 @@ namespace UE::ConcertSyncClient::Replication
 			const TOptional<FConcertReplication_ChangeAuthority_Request> AuthorityChange;
 			const TMap<FSoftObjectPath, TArray<FGuid>> PredictedStreamChange;
 			const FLocalSyncControl::FPredictedObjectRemoval PredictedMuteChange;
+			
 			FConcertReplication_ChangeSyncControl RetainedSyncControl;
+			FChangeStreamPredictedChange StreamChangeToRevert;
+			TMap<FSoftObjectPath, TArray<FGuid>> AuthorityChangeToRevert;
 
 			FPutStateChange(
 				const TConstArrayView<FConcertReplicationStream>& RegisteredStreams,
@@ -128,7 +131,6 @@ namespace UE::ConcertSyncClient::Replication
 				)
 				: StreamChange(BuildDeltaStreamChange(ClientId, RegisteredStreams, Request))
 				, AuthorityChange(BuildDeltaAuthorityChange(ClientId, Request))
-				, PredictedStreamChange(StreamChange ? ComputeRemovedObjects(RegisteredStreams, *StreamChange) : TMap<FSoftObjectPath, TArray<FGuid>>{})
 				, PredictedMuteChange(MoveTemp(PredictedMuteChange))
 				, RetainedSyncControl(BuildRetainedSyncControl(ClientId, SyncControl, Request))
 			{}
@@ -150,11 +152,11 @@ namespace UE::ConcertSyncClient::Replication
 		FPutStateChange Change(RegisteredStreams, MoveTemp(PredictedMuteChange), SyncControl, ClientId, Request);
 		if (Change.StreamChange)
 		{
-			ApplyStreamChangeRemovedObjects(Change.PredictedStreamChange);
+			Change.StreamChangeToRevert = PredictAndApplyStreamChangeRemovedObjects(*Change.StreamChange);
 		}
 		if (Change.AuthorityChange)
 		{
-			ApplyAuthorityChangeRemovedObjects(*Change.AuthorityChange);
+			Change.AuthorityChangeToRevert = ApplyAuthorityChangeRemovedObjects(*Change.AuthorityChange);
 		}
 		
 		LogNetworkMessage(CVarLogChangeClientsRequestsAndResponsesOnClient, Request);
@@ -193,13 +195,11 @@ namespace UE::ConcertSyncClient::Replication
 					// [1]... and we may have to revert predictive changes
 					if (Change.StreamChange)
 					{
-						// HandleRemovingReplicatedObjects caused Request.ObjectsToRemove to stop being replicated. Revert.
-						ThisPin->RevertPredictedStreamChangeRemovedObjects(Change.PredictedStreamChange);
+						ThisPin->RevertPredictedStreamChangeRemovedObjects(Change.StreamChangeToRevert);
 					}
 					if (Change.AuthorityChange)
 					{
-						// HandleReleasingReplicatedObjects caused Args.ReleaseAuthority to stop being replicated. Revert.
-						ThisPin->RevertAuthorityChangeReleasedObjects(*Change.AuthorityChange);
+						ThisPin->RevertAuthorityChangeReleasedObjects(Change.AuthorityChangeToRevert);
 					}
 					
 					ThisPin->SyncControl.ApplyOrRevertMuteResponse(Change.PredictedMuteChange, { .ErrorCode = EConcertReplicationMuteErrorCode::Rejected });
diff --git a/Engine/Plugins/Developer/Concert/ConcertSync/ConcertSyncClient/Source/ConcertSyncClient/Private/Replication/Processing/ClientReplicationDataCollector.cpp b/Engine/Plugins/Developer/Concert/ConcertSync/ConcertSyncClient/Source/ConcertSyncClient/Private/Replication/Processing/ClientReplicationDataCollector.cpp
--- a/Engine/Plugins/Developer/Concert/ConcertSync/ConcertSyncClient/Source/ConcertSyncClient/Private/Replication/Processing/ClientReplicationDataCollector.cpp
+++ b/Engine/Plugins/Developer/Concert/ConcertSync/ConcertSyncClient/Source/ConcertSyncClient/Private/Replication/Processing/ClientReplicationDataCollector.cpp
@@ -45,6 +45,11 @@ namespace UE::ConcertSyncClient::Replication
 	
 	void FClientReplicationDataCollector::AddReplicatedObjectStreams(const FSoftObjectPath& Object, TArrayView<const FGuid> AddedStreams)
 	{
+		if (AddedStreams.IsEmpty())
+		{
+			return;
+		}
+		
 		TArray<FObjectInfo>& ReplicatedObjectInfo = ObjectsToReplicate.FindOrAdd(Object);
 		const bool bIsNewReplicatedObject = ReplicatedObjectInfo.IsEmpty();
 		
diff --git a/Engine/Plugins/Developer/Concert/ConcertSync/ConcertSyncTest/Source/ConcertSyncTest/Private/Replication/Bugs/RevertChange.spec.cpp b/Engine/Plugins/Developer/Concert/ConcertSync/ConcertSyncTest/Source/ConcertSyncTest/Private/Replication/Bugs/RevertChange.spec.cpp
new file mode 100644
--- /dev/null
+++ b/Engine/Plugins/Developer/Concert/ConcertSync/ConcertSyncTest/Source/ConcertSyncTest/Private/Replication/Bugs/RevertChange.spec.cpp
@@ -0,0 +1,62 @@
+﻿// Copyright Epic Games, Inc. All Rights Reserved.
+
+#include "Misc/AutomationTest.h"
+
+#include "Replication/Util/Spec/ObjectTestReplicator.h"
+#include "Replication/Util/Spec/ReplicationClient.h"
+#include "Replication/Util/Spec/ReplicationServer.h"
+
+namespace UE::ConcertSyncTests::Replication::UI
+{
+	BEGIN_DEFINE_SPEC(FRevertChange, "Editor.Concert.Replication.Bugs.RevertRequest", EAutomationTestFlags::EditorContext | EAutomationTestFlags::EngineFilter)
+		TUniquePtr<FObjectTestReplicator> Replicator;
+		TUniquePtr<FReplicationServer> Server;
+		FReplicationClient* Client1;
+
+		const FGuid StreamId = FGuid::NewGuid();
+	END_DEFINE_SPEC(FRevertChange);
+
+	/** There used to be a bug where the client would revert a failed authority request incorrectly ending up thinking it has authority even though it does not. */
+	void FRevertChange::Define()
+	{
+		BeforeEach([this]
+		{
+			Replicator = MakeUnique<FObjectTestReplicator>();
+			Server = MakeUnique<FReplicationServer>(*this);
+			Client1 = &Server->ConnectClient();
+
+			Client1->JoinReplication(Replicator->CreateSenderArgs(StreamId));
+		});
+		AfterEach([this]
+		{
+			Replicator.Reset();
+			Server.Reset();
+		});
+
+		It("Reverting a stream change operation does not grant authority if the request fails", [this]
+		{
+			// Simulate timeout
+			const TSharedRef<FConcertServerSessionMock>& Session = Server->GetServerSessionMock();
+			Session->SetTestFlags(EServerSessionTestingFlags::AllowRequestTimeouts);
+			Session->UnregisterCustomRequestHandler<FConcertReplication_ChangeStream_Request>();
+
+			// The revert mechanism used to blindly add back authority for each removed object instead of checking whether it actually had authority when the request was made.
+			IConcertClientReplicationManager& ReplicationManager = Client1->GetClientReplicationManager();
+			ReplicationManager.ChangeStream({ .ObjectsToRemove = { { StreamId, Replicator->TestObject  } }});
+			TestFalse(TEXT("Receive authority"), ReplicationManager.HasAuthorityOver(Replicator->TestObject));
+		});
+		
+		It("Reverting a release authority operation does not grant authority if the request fails", [this]
+		{
+			// Simulate timeout
+			const TSharedRef<FConcertServerSessionMock>& Session = Server->GetServerSessionMock();
+			Session->SetTestFlags(EServerSessionTestingFlags::AllowRequestTimeouts);
+			Session->UnregisterCustomRequestHandler<FConcertReplication_ChangeAuthority_Request>();
+
+			// The revert mechanism used to blindly add back authority for each requested object instead of checking whether it actually had authority when the request was made.
+			IConcertClientReplicationManager& ReplicationManager = Client1->GetClientReplicationManager();
+			ReplicationManager.ReleaseAuthorityOf({ Replicator->TestObject });
+			TestFalse(TEXT("Receive authority"), ReplicationManager.HasAuthorityOver(Replicator->TestObject));
+		});
+	}
+}
\ No newline at end of file
