OnClientAuthorityServerStateChanged  now triggers when deleting an object locally

diff --git a/Engine/Plugins/Developer/Concert/ConcertSync/ConcertSyncClient/Source/ConcertSyncClient/Private/Replication/Manager/ReplicationManagerState_Connected.cpp b/Engine/Plugins/Developer/Concert/ConcertSync/ConcertSyncClient/Source/ConcertSyncClient/Private/Replication/Manager/ReplicationManagerState_Connected.cpp
--- a/Engine/Plugins/Developer/Concert/ConcertSync/ConcertSyncClient/Source/ConcertSyncClient/Private/Replication/Manager/ReplicationManagerState_Connected.cpp
+++ b/Engine/Plugins/Developer/Concert/ConcertSync/ConcertSyncClient/Source/ConcertSyncClient/Private/Replication/Manager/ReplicationManagerState_Connected.cpp
@@ -23,59 +23,6 @@
 #include "UObject/UObjectGlobals.h"
 #include "Utils/ContentComparisionUtils.h"
 
-namespace UE::ConcertSyncClient::Replication::Private
-{
-	/** Removes ObjectsToRemove and returns the authority that was actually removed, possibly to be reverted later. */
-	template<typename TStreamArrayType, typename TProjectStreamArray> requires std::is_invocable_r_v<const TArray<FGuid>&, TProjectStreamArray, const TStreamArrayType&>
-	static TMap<FSoftObjectPath, TArray<FGuid>> RemoveObjectsAndTrackRemovedAuthority(
-		FClientReplicationDataCollector& ReplicationDataSource,
-		const TMap<FSoftObjectPath, TStreamArrayType>& ObjectsToRemove,
-		TProjectStreamArray&& Project
-		)
-	{
-		TMap<FSoftObjectPath, TArray<FGuid>> RevertedContent;
-		for (const TPair<FSoftObjectPath, TStreamArrayType>& ReleaseAuthority : ObjectsToRemove)
-		{
-			const FSoftObjectPath& ObjectPath = ReleaseAuthority.Key;
-			const TArray<FGuid>& StreamsToRemove = Project(ReleaseAuthority.Value);
-
-			// Remember which streams we actually removed, so we can revert it again.
-			TArray<FGuid> StreamsActuallyRemoved = ReplicationDataSource.GetStreamsOwningObject(ObjectPath);
-			StreamsActuallyRemoved.SetNum(Algo::RemoveIf(StreamsActuallyRemoved, [&StreamsToRemove](const FGuid& StreamId)
-			{
-				return !StreamsToRemove.Contains(StreamId);
-			}));
-			RevertedContent.Add(ObjectPath, StreamsActuallyRemoved);
-			
-			ReplicationDataSource.RemoveReplicatedObjectStreams(ObjectPath, StreamsToRemove);
-		}
-		return RevertedContent;
-	}
-
-	static TMap<FSoftObjectPath, TArray<FGuid>> RemoveObjectsAndTrackRemovedAuthority(
-		FClientReplicationDataCollector& ReplicationDataSource,
-		const TMap<FSoftObjectPath, FConcertStreamArray>& ObjectsToRemove
-		)
-	{
-		return RemoveObjectsAndTrackRemovedAuthority<FConcertStreamArray>(
-			ReplicationDataSource, ObjectsToRemove,
-			[](const FConcertStreamArray& StreamArray)-> const TArray<FGuid>& { return StreamArray.StreamIds; }
-			); 
-	}
-	static TMap<FSoftObjectPath, TArray<FGuid>> RemoveObjectsAndTrackRemovedAuthority(
-		FClientReplicationDataCollector& ReplicationDataSource,
-		const TMap<FSoftObjectPath, TArray<FGuid>>& ObjectsToRemove
-		)
-	{
-		return RemoveObjectsAndTrackRemovedAuthority<TArray<FGuid>>(
-			ReplicationDataSource, ObjectsToRemove,
-			[](const TArray<FGuid>& StreamArray)-> const TArray<FGuid>& { return StreamArray; }
-			); 
-	}
-	
-	
-}
-
 namespace UE::ConcertSyncClient::Replication
 {
 	FReplicationManagerState_Connected::FReplicationManagerState_Connected(
@@ -382,23 +329,23 @@ namespace UE::ConcertSyncClient::Replication
 	{
 		FChangeStreamPredictedChange Change;
 		Change.ObjectsRemovedFromStream = ComputeRemovedObjects(RegisteredStreams, Request);
-		Change.AuthorityRemovedFromStreams = ApplyStreamChangeRemovedObjects(Change.ObjectsRemovedFromStream);
+		Change.AuthorityRemovedFromStreams = RemoveObjectsFromAuthority(Change.ObjectsRemovedFromStream);
 		return Change;
 	}
 
-	TMap<FSoftObjectPath, TArray<FGuid>> FReplicationManagerState_Connected::ApplyStreamChangeRemovedObjects(
-		const TMap<FSoftObjectPath, TArray<FGuid>>& PredicatedRemovedObjects
-		)
-	{
-		return Private::RemoveObjectsAndTrackRemovedAuthority(ReplicationDataSource, PredicatedRemovedObjects);
-	}
-
 	void FReplicationManagerState_Connected::RevertPredictedStreamChangeRemovedObjects(const FChangeStreamPredictedChange& PredictedChange)
 	{
+		if (PredictedChange.AuthorityRemovedFromStreams.IsEmpty())
+		{
+			return;
+		}
+
+		OnPreAuthorityChangedDelegate.Broadcast();
 		for (const TPair<FSoftObjectPath, TArray<FGuid>>& RemovedObjectInfo : PredictedChange.AuthorityRemovedFromStreams)
 		{
 			ReplicationDataSource.AddReplicatedObjectStreams(RemovedObjectInfo.Key, RemovedObjectInfo.Value);
 		}
+		OnPostAuthorityChangedDelegate.Broadcast();
 	}
 	
 	void FReplicationManagerState_Connected::FinalizePredictedStreamChange(const FConcertReplication_ChangeStream_Request& StreamChange)
@@ -442,26 +389,37 @@ namespace UE::ConcertSyncClient::Replication
 		}
 	}
 
-	TMap<FSoftObjectPath, TArray<FGuid>> FReplicationManagerState_Connected::ApplyAuthorityChangeRemovedObjects(const FConcertReplication_ChangeAuthority_Request& Request)
+	TMap<FSoftObjectPath, TArray<FGuid>> FReplicationManagerState_Connected::ApplyAuthorityChangeRemovedObjects(
+		const FConcertReplication_ChangeAuthority_Request& Request
+		)
 	{
-		return Private::RemoveObjectsAndTrackRemovedAuthority(ReplicationDataSource, Request.ReleaseAuthority);
+		return RemoveObjectsFromAuthority(Request.ReleaseAuthority);
 	}
 
 	void FReplicationManagerState_Connected::RevertAuthorityChangeReleasedObjects(const TMap<FSoftObjectPath, TArray<FGuid>>& PredictedChange)
 	{
+		OnPreAuthorityChangedDelegate.Broadcast();
 		for (const TPair<FSoftObjectPath, TArray<FGuid>>& ReleaseAuthority : PredictedChange)
 		{
 			ReplicationDataSource.AddReplicatedObjectStreams(ReleaseAuthority.Key, ReleaseAuthority.Value);
 		}
+		OnPostAuthorityChangedDelegate.Broadcast();
 	}
 
-	void FReplicationManagerState_Connected::FinalizePredictedAuthorityChange(const FConcertReplication_ChangeAuthority_Request& AuthorityChange, const TMap<FSoftObjectPath, FConcertStreamArray>& RejectedObjects, const FConcertReplication_ChangeSyncControl& SyncControlChange)
+	void FReplicationManagerState_Connected::FinalizePredictedAuthorityChange(
+		const FConcertReplication_ChangeAuthority_Request& AuthorityChange,
+		const TMap<FSoftObjectPath, FConcertStreamArray>& RejectedObjects,
+		const FConcertReplication_ChangeSyncControl& SyncControlChange
+		)
 	{
 		SyncControl.ProcessAuthorityChange(AuthorityChange, SyncControlChange);
 		UpdateReplicatedObjectsAfterAuthorityChange(AuthorityChange, RejectedObjects);
 	}
 	
-	void FReplicationManagerState_Connected::UpdateReplicatedObjectsAfterAuthorityChange(const FConcertReplication_ChangeAuthority_Request& Request, const TMap<FSoftObjectPath, FConcertStreamArray>& RejectedObjects)
+	void FReplicationManagerState_Connected::UpdateReplicatedObjectsAfterAuthorityChange(
+		const FConcertReplication_ChangeAuthority_Request& Request,
+		const TMap<FSoftObjectPath, FConcertStreamArray>& RejectedObjects
+		)
 	{
 		OnPreAuthorityChangedDelegate.Broadcast();
 		ON_SCOPE_EXIT{ OnPostAuthorityChangedDelegate.Broadcast(); };
diff --git a/Engine/Plugins/Developer/Concert/ConcertSync/ConcertSyncClient/Source/ConcertSyncClient/Private/Replication/Manager/ReplicationManagerState_Connected.h b/Engine/Plugins/Developer/Concert/ConcertSync/ConcertSyncClient/Source/ConcertSyncClient/Private/Replication/Manager/ReplicationManagerState_Connected.h
--- a/Engine/Plugins/Developer/Concert/ConcertSync/ConcertSyncClient/Source/ConcertSyncClient/Private/Replication/Manager/ReplicationManagerState_Connected.h
+++ b/Engine/Plugins/Developer/Concert/ConcertSync/ConcertSyncClient/Source/ConcertSyncClient/Private/Replication/Manager/ReplicationManagerState_Connected.h
@@ -12,8 +12,11 @@
 #include "Replication/Processing/ObjectReplicationReceiver.h"
 #include "Replication/Processing/ObjectReplicationSender.h"
 #include "Replication/Processing/Proxy/ObjectProcessorProxy_Frequency.h"
+#include "Utils/AuthorityRemovalPrediction.h"
 #include "Utils/LocalSyncControl.h"
 
+#include <type_traits>
+
 class IConcertClientSession;
 
 namespace UE::ConcertSyncCore { class FObjectReplicationReceiver; }
@@ -124,7 +127,6 @@ namespace UE::ConcertSyncClient::Replication
 
 		/** Changes the local state assuming that Request will succeed. */
 		FChangeStreamPredictedChange PredictAndApplyStreamChangeRemovedObjects(const FConcertReplication_ChangeStream_Request& Request);
-		TMap<FSoftObjectPath, TArray<FGuid>> ApplyStreamChangeRemovedObjects(const TMap<FSoftObjectPath, TArray<FGuid>>& PredicatedRemovedObjects);
 		/** Reverts changes previously made by PredictStreamChangeRemovedObjects. */
 		void RevertPredictedStreamChangeRemovedObjects(const FChangeStreamPredictedChange& PredictedChange);
 		/** Applies stream changes that we previously predicted using PredictStreamChangeRemovedObjects. */
@@ -141,6 +143,13 @@ namespace UE::ConcertSyncClient::Replication
 		/** Updates the objects which should be replicated after changing authority. */
 		void UpdateReplicatedObjectsAfterAuthorityChange(const FConcertReplication_ChangeAuthority_Request& Request, const TMap<FSoftObjectPath, FConcertStreamArray>& RejectedObjects);
 
+		/**
+		 * Shared logic for removing objects from authority.
+		 * @return The objects that were actually removed (so it can be reverted).
+		 */
+		template<typename TStreamIdArray> requires std::is_same_v<TStreamIdArray, TArray<FGuid>> || std::is_same_v<TStreamIdArray, FConcertStreamArray>
+		TMap<FSoftObjectPath, TArray<FGuid>> RemoveObjectsFromAuthority(const TMap<FSoftObjectPath, TStreamIdArray>& PredicatedObjects); 
+
 		/** Updates the objects which should be replicated after they have been reset to a completely new state (e.g. when restoring session content manually). */
 		void UpdateReplicatedObjectAfterServerSideChange(const FConcertQueriedClientInfo& NewState);
 		
@@ -148,3 +157,21 @@ namespace UE::ConcertSyncClient::Replication
 		FConcertObjectReplicationSettings GetObjectFrequencySettings(const FConcertReplicatedObjectId& Object) const;
 	};
 }
+
+namespace UE::ConcertSyncClient::Replication
+{
+	template <typename TStreamIdArray> requires std::is_same_v<TStreamIdArray, TArray<FGuid>> || std::is_same_v<TStreamIdArray, FConcertStreamArray>
+	TMap<FSoftObjectPath, TArray<FGuid>> FReplicationManagerState_Connected::RemoveObjectsFromAuthority(
+		const TMap<FSoftObjectPath, TStreamIdArray>& PredicatedObjects
+	)
+	{
+		if (!HasAuthorityChanges(ReplicationDataSource, PredicatedObjects))
+		{
+			return {};
+		}
+		
+		OnPreAuthorityChangedDelegate.Broadcast();
+		ON_SCOPE_EXIT{ OnPostAuthorityChangedDelegate.Broadcast(); };
+		return RemoveObjectsAndTrackRemovedAuthority(ReplicationDataSource, PredicatedObjects);
+	}
+}
diff --git a/Engine/Plugins/Developer/Concert/ConcertSync/ConcertSyncClient/Source/ConcertSyncClient/Private/Replication/Manager/Utils/AuthorityRemovalPrediction.h b/Engine/Plugins/Developer/Concert/ConcertSync/ConcertSyncClient/Source/ConcertSyncClient/Private/Replication/Manager/Utils/AuthorityRemovalPrediction.h
new file mode 100644
--- /dev/null
+++ b/Engine/Plugins/Developer/Concert/ConcertSync/ConcertSyncClient/Source/ConcertSyncClient/Private/Replication/Manager/Utils/AuthorityRemovalPrediction.h
@@ -0,0 +1,130 @@
+﻿// Copyright Epic Games, Inc. All Rights Reserved.
+
+#pragma once
+
+#include "Replication/Messages/ChangeAuthority.h"
+#include "Replication/Processing/ClientReplicationDataCollector.h"
+
+#include "Algo/AnyOf.h"
+#include "Algo/RemoveIf.h"
+#include "Containers/Array.h"
+#include "Containers/Map.h"
+#include "Misc/Guid.h"
+#include "UObject/SoftObjectPath.h"
+
+#include <type_traits>
+
+namespace UE::ConcertSyncClient::Replication
+{
+	class FClientReplicationDataCollector;
+	
+	/** Removes ObjectsToRemove and returns the authority that was actually removed, possibly to be reverted later. */
+	template<typename TStreamArrayType, typename TProjectStreamArray> requires std::is_invocable_r_v<const TArray<FGuid>&, TProjectStreamArray, const TStreamArrayType&>
+	TMap<FSoftObjectPath, TArray<FGuid>> RemoveObjectsAndTrackRemovedAuthorityGeneric(
+		FClientReplicationDataCollector& ReplicationDataSource,
+		const TMap<FSoftObjectPath, TStreamArrayType>& ObjectsToRemove,
+		TProjectStreamArray&& Project
+		);
+
+	/** @return Whether RemoveObjectsAndTrackRemovedAuthority would make any changes. This is useful for avoiding broadcasting change events. */
+	template<typename TStreamArrayType, typename TProjectStreamArray> requires std::is_invocable_r_v<const TArray<FGuid>&, TProjectStreamArray, const TStreamArrayType&>
+	bool HasAuthorityChangesGeneric(
+		FClientReplicationDataCollector& ReplicationDataSource,
+		const TMap<FSoftObjectPath, TStreamArrayType>& ObjectsToRemove,
+		TProjectStreamArray&& Project
+	);
+	
+	// Inline version for FConcertStreamArray & TArray<FGuid>
+	TMap<FSoftObjectPath, TArray<FGuid>> RemoveObjectsAndTrackRemovedAuthority(
+		FClientReplicationDataCollector& ReplicationDataSource, const TMap<FSoftObjectPath, FConcertStreamArray>& ObjectsToRemove
+		);
+	TMap<FSoftObjectPath, TArray<FGuid>> RemoveObjectsAndTrackRemovedAuthority(
+		FClientReplicationDataCollector& ReplicationDataSource, const TMap<FSoftObjectPath, TArray<FGuid>>& ObjectsToRemove
+		);
+	bool HasAuthorityChanges(FClientReplicationDataCollector& ReplicationDataSource, const TMap<FSoftObjectPath, TArray<FGuid>>& ObjectsToRemove);
+	bool HasAuthorityChanges(FClientReplicationDataCollector& ReplicationDataSource, const TMap<FSoftObjectPath, FConcertStreamArray>& ObjectsToRemove);
+}
+
+namespace UE::ConcertSyncClient::Replication
+{
+	/** Removes ObjectsToRemove and returns the authority that was actually removed, possibly to be reverted later. */
+	template<typename TStreamArrayType, typename TProjectStreamArray> requires std::is_invocable_r_v<const TArray<FGuid>&, TProjectStreamArray, const TStreamArrayType&>
+	TMap<FSoftObjectPath, TArray<FGuid>> RemoveObjectsAndTrackRemovedAuthorityGeneric(
+		FClientReplicationDataCollector& ReplicationDataSource,
+		const TMap<FSoftObjectPath, TStreamArrayType>& ObjectsToRemove,
+		TProjectStreamArray&& Project
+		)
+	{
+		TMap<FSoftObjectPath, TArray<FGuid>> RevertedContent;
+		for (const TPair<FSoftObjectPath, TStreamArrayType>& ReleaseAuthority : ObjectsToRemove)
+		{
+			const FSoftObjectPath& ObjectPath = ReleaseAuthority.Key;
+			const TArray<FGuid>& StreamsToRemove = Project(ReleaseAuthority.Value);
+
+			// Remember which streams we actually removed, so we can revert it again.
+			TArray<FGuid> StreamsActuallyRemoved = ReplicationDataSource.GetStreamsOwningObject(ObjectPath);
+			StreamsActuallyRemoved.SetNum(Algo::RemoveIf(StreamsActuallyRemoved, [&StreamsToRemove](const FGuid& StreamId)
+			{
+				return !StreamsToRemove.Contains(StreamId);
+			}));
+			RevertedContent.Add(ObjectPath, StreamsActuallyRemoved);
+			
+			ReplicationDataSource.RemoveReplicatedObjectStreams(ObjectPath, StreamsToRemove);
+		}
+		return RevertedContent;
+	}
+
+	/** @return Whether RemoveObjectsAndTrackRemovedAuthority would make any changes. This is useful for avoiding broadcasting change events. */
+	template<typename TStreamArrayType, typename TProjectStreamArray> requires std::is_invocable_r_v<const TArray<FGuid>&, TProjectStreamArray, const TStreamArrayType&>
+	bool HasAuthorityChangesGeneric(
+		FClientReplicationDataCollector& ReplicationDataSource,
+		const TMap<FSoftObjectPath, TStreamArrayType>& ObjectsToRemove,
+		TProjectStreamArray&& Project
+	)
+	{
+		return Algo::AnyOf(ObjectsToRemove, [&ReplicationDataSource, &Project](const TPair<FSoftObjectPath, TStreamArrayType>& Pair)
+		{
+			const FSoftObjectPath& ObjectPath = Pair.Key;
+			return Algo::AnyOf(Project(Pair.Value), [&ObjectPath, &ReplicationDataSource](const FGuid& StreamId)
+			{
+				return ReplicationDataSource.OwnsObjectInStream(ObjectPath, StreamId);
+			});
+		});
+	}
+
+	inline TMap<FSoftObjectPath, TArray<FGuid>> RemoveObjectsAndTrackRemovedAuthority(
+		FClientReplicationDataCollector& ReplicationDataSource,
+		const TMap<FSoftObjectPath, FConcertStreamArray>& ObjectsToRemove
+		)
+	{
+		return RemoveObjectsAndTrackRemovedAuthorityGeneric<FConcertStreamArray>(
+			ReplicationDataSource, ObjectsToRemove,
+			[](const FConcertStreamArray& StreamArray)-> const TArray<FGuid>& { return StreamArray.StreamIds; }
+			); 
+	}
+
+	inline TMap<FSoftObjectPath, TArray<FGuid>> RemoveObjectsAndTrackRemovedAuthority(
+		FClientReplicationDataCollector& ReplicationDataSource,
+		const TMap<FSoftObjectPath, TArray<FGuid>>& ObjectsToRemove
+		)
+	{
+		return RemoveObjectsAndTrackRemovedAuthorityGeneric<TArray<FGuid>>(
+			ReplicationDataSource, ObjectsToRemove,
+			[](const TArray<FGuid>& StreamArray)-> const TArray<FGuid>& { return StreamArray; }
+			); 
+	}
+
+	inline bool HasAuthorityChanges(FClientReplicationDataCollector& ReplicationDataSource, const TMap<FSoftObjectPath, TArray<FGuid>>& ObjectsToRemove)
+	{
+		return HasAuthorityChangesGeneric<TArray<FGuid>>(ReplicationDataSource, ObjectsToRemove,
+			[](const TArray<FGuid>& StreamArray)-> const TArray<FGuid>& { return StreamArray; }
+			);
+	}
+
+	inline bool HasAuthorityChanges(FClientReplicationDataCollector& ReplicationDataSource, const TMap<FSoftObjectPath, FConcertStreamArray>& ObjectsToRemove)
+	{
+		return HasAuthorityChangesGeneric<FConcertStreamArray>(ReplicationDataSource, ObjectsToRemove,
+			[](const FConcertStreamArray& StreamArray)-> const TArray<FGuid>& { return StreamArray.StreamIds; }
+			);
+	}
+}
\ No newline at end of file
diff --git a/Engine/Plugins/Developer/Concert/ConcertSync/ConcertSyncTest/Source/ConcertSyncTest/Private/Replication/ClientManager/ChangeStreamAndAuthorityEvents.spec.cpp b/Engine/Plugins/Developer/Concert/ConcertSync/ConcertSyncTest/Source/ConcertSyncTest/Private/Replication/ClientManager/ChangeStreamAndAuthorityEvents.spec.cpp
new file mode 100644
--- /dev/null
+++ b/Engine/Plugins/Developer/Concert/ConcertSync/ConcertSyncTest/Source/ConcertSyncTest/Private/Replication/ClientManager/ChangeStreamAndAuthorityEvents.spec.cpp
@@ -0,0 +1,69 @@
+﻿// Copyright Epic Games, Inc. All Rights Reserved.
+
+#include "Misc/AutomationTest.h"
+
+#include "Replication/Util/Spec/ClientEventCounter.h"
+#include "Replication/Util/Spec/ObjectTestReplicator.h"
+#include "Replication/Util/Spec/ReplicationClient.h"
+#include "Replication/Util/Spec/ReplicationServer.h"
+
+namespace UE::ConcertSyncTests::Replication::UI
+{
+	BEGIN_DEFINE_SPEC(FChangeStreamAndAuthorityEvents, "Editor.Concert.Replication.ClientManager.ChangeStreamAndAuthorityEvents", EAutomationTestFlags::EditorContext | EAutomationTestFlags::EngineFilter)
+		TUniquePtr<FObjectTestReplicator> Replicator;
+		TUniquePtr<FReplicationServer> Server;
+		FReplicationClient* Sender;
+
+		const FGuid StreamId = FGuid::NewGuid();
+	
+		/** Will help us count the client events. */
+		FClientEventCounter EventCounter;
+	END_DEFINE_SPEC(FChangeStreamAndAuthorityEvents);
+
+	/**
+	 * This tests that IConcertClientReplicationManager stream and authority change events are made correctly in response to
+	 * IConcertClientReplicationManager::ChangeStream.
+	 */
+	void FChangeStreamAndAuthorityEvents::Define()
+	{
+		BeforeEach([this]
+		{
+			Replicator = MakeUnique<FObjectTestReplicator>();
+			Server = MakeUnique<FReplicationServer>(*this);
+			Sender = &Server->ConnectClient();
+
+			Sender->JoinReplication();
+			EventCounter.Subscribe(*Sender);
+		});
+		AfterEach([this]
+		{
+			Replicator.Reset();
+			Server.Reset();
+		});
+
+		It("When adding a stream, broadcast stream event", [this]
+		{
+			Sender->GetClientReplicationManager().ChangeStream({ .StreamsToAdd = { Replicator->CreateStream(StreamId) } });
+			EventCounter.TestCount(*this, 1, 0);
+		});
+
+		It("When taking authority, broadcast authority event", [this]
+		{
+			Sender->GetClientReplicationManager().ChangeStream({ .StreamsToAdd = { Replicator->CreateStream(StreamId) } });
+			EventCounter.ResetEventCount();
+			
+			Sender->GetClientReplicationManager().TakeAuthorityOver({ Replicator->TestObject });
+			EventCounter.TestCount(*this, 0, 1);
+		});
+		
+		It("When removing owned object, broadcast stream and authority event", [this]
+		{
+			Sender->GetClientReplicationManager().ChangeStream({ .StreamsToAdd = { Replicator->CreateStream(StreamId) } });
+			Sender->GetClientReplicationManager().TakeAuthorityOver({ Replicator->TestObject });
+			EventCounter.ResetEventCount();
+
+			Sender->GetClientReplicationManager().ChangeStream({ .StreamsToRemove = { StreamId } });
+			EventCounter.TestCount(*this, 1, 1);
+		});
+	}
+}
\ No newline at end of file
diff --git a/Engine/Plugins/Developer/Concert/ConcertSync/ConcertSyncTest/Source/ConcertSyncTest/Private/Replication/ClientManager/PutClientStateEvents.spec.cpp b/Engine/Plugins/Developer/Concert/ConcertSync/ConcertSyncTest/Source/ConcertSyncTest/Private/Replication/ClientManager/PutClientStateEvents.spec.cpp
new file mode 100644
--- /dev/null
+++ b/Engine/Plugins/Developer/Concert/ConcertSync/ConcertSyncTest/Source/ConcertSyncTest/Private/Replication/ClientManager/PutClientStateEvents.spec.cpp
@@ -0,0 +1,79 @@
+﻿// Copyright Epic Games, Inc. All Rights Reserved.
+
+#include "Misc/AutomationTest.h"
+
+#include "Replication/Util/Spec/ClientEventCounter.h"
+#include "Replication/Util/Spec/ObjectTestReplicator.h"
+#include "Replication/Util/Spec/ReplicationClient.h"
+#include "Replication/Util/Spec/ReplicationServer.h"
+
+namespace UE::ConcertSyncTests::Replication::UI
+{
+	BEGIN_DEFINE_SPEC(FPutClientStateEvents, "Editor.Concert.Replication.ClientManager.PutClientStateEvents", EAutomationTestFlags::EditorContext | EAutomationTestFlags::EngineFilter)
+		TUniquePtr<FObjectTestReplicator> Replicator_Client1;
+		TUniquePtr<FObjectTestReplicator> Replicator_Client2;
+		TUniquePtr<FReplicationServer> Server;
+		FReplicationClient* Client1;
+		FReplicationClient* Client2;
+
+		const FGuid StreamId = FGuid::NewGuid();
+	
+		/** Will help us count the client events. */
+		FClientEventCounter EventCounter_Client1;
+		/** Will help us count the client events. */
+		FClientEventCounter EventCounter_Client2;
+	
+		FConcertReplication_PutState_Request BuildBaseRequest() const
+		{
+			FConcertReplication_PutState_Request Request;
+
+			FConcertReplicationStreamArray StreamContent_Client1;
+			StreamContent_Client1.Streams = { Replicator_Client1->CreateStream(StreamId) };
+			Request.NewStreams.Add(Client1->GetEndpointId(), StreamContent_Client1);
+			
+			FConcertReplicationStreamArray StreamContent_Client2;
+			StreamContent_Client2.Streams = { Replicator_Client2->CreateStream(StreamId) };
+			Request.NewStreams.Add(Client2->GetEndpointId(), StreamContent_Client2);
+
+			const FConcertObjectInStreamID ObjectId_Client1 = { StreamId, Replicator_Client1->TestObject };
+			const FConcertObjectInStreamID ObjectId_Client2 = { StreamId, Replicator_Client2->TestObject };
+			Request.NewAuthorityState.Add(Client1->GetEndpointId(), { { ObjectId_Client1 } });
+			Request.NewAuthorityState.Add(Client2->GetEndpointId(), { { ObjectId_Client2 } });
+				
+			return Request;
+		}
+	END_DEFINE_SPEC(FPutClientStateEvents);
+
+	/**
+	 * This tests that IConcertClientReplicationManager stream and authority change events are made correctly in response to
+	 * IConcertClientReplicationManager::PutClientState.
+	 */
+	void FPutClientStateEvents::Define()
+	{
+		BeforeEach([this]
+		{
+			Replicator_Client1 = MakeUnique<FObjectTestReplicator>(TEXT("Foo"));
+			Replicator_Client2 = MakeUnique<FObjectTestReplicator>(TEXT("Bar"));
+			Server = MakeUnique<FReplicationServer>(*this);
+			Client1 = &Server->ConnectClient();
+			Client2 = &Server->ConnectClient();
+
+			Client1->JoinReplication();
+			Client2->JoinReplication();
+			EventCounter_Client1.Subscribe(*Client1);
+			EventCounter_Client2.Subscribe(*Client2);
+		});
+		AfterEach([this]
+		{
+			Replicator_Client1.Reset();
+			Server.Reset();
+		});
+
+		It("When PutState is applied on an empty client, broadcast stream and authority event once.", [this]
+		{
+			Client1->GetClientReplicationManager().PutClientState(BuildBaseRequest());
+			EventCounter_Client1.TestCount(*this, 1, 1);
+			EventCounter_Client2.TestCount(*this, 1, 1);
+		});
+	}
+}
\ No newline at end of file
