Add some additional logging to ConcertSandboxPlatformFile.cpp

diff --git a/Engine/Plugins/Developer/Concert/ConcertSync/ConcertSyncClient/Source/ConcertSyncClient/Private/ConcertSandboxPlatformFile.cpp b/Engine/Plugins/Developer/Concert/ConcertSync/ConcertSyncClient/Source/ConcertSyncClient/Private/ConcertSandboxPlatformFile.cpp
--- a/Engine/Plugins/Developer/Concert/ConcertSync/ConcertSyncClient/Source/ConcertSyncClient/Private/ConcertSandboxPlatformFile.cpp
+++ b/Engine/Plugins/Developer/Concert/ConcertSync/ConcertSyncClient/Source/ConcertSyncClient/Private/ConcertSandboxPlatformFile.cpp
@@ -1212,16 +1212,19 @@ void FConcertSandboxPlatformFile::MigrateFileToSandbox(const FConcertSandboxPlat
 
 	if (LowerLevel->FileExists(*InPath.GetSandboxPath()))
 	{
+		UE_LOG(LogConcert, Verbose, TEXT("MigrateFileToSandbox: Sandbox has the file %s already at this location. Nothing to do."), *InPath.GetSandboxPath() );
 		// Sandbox already has a file at this location - nothing to do
 		return;
 	}
 
 	if (!LowerLevel->FileExists(*InPath.GetNonSandboxPath()))
 	{
+		UE_LOG(LogConcert, Verbose, TEXT("MigrateFileToSandbox: There is no file %s at this location. Nothing to do."), *InPath.GetNonSandboxPath() );
 		// Non-sandbox has no file at this location - nothing to do
 		return;
 	}
 
+	UE_LOG(LogConcert, Verbose, TEXT("MigrateFileToSandbox: Copying file %s into the sandbox."), *InPath.GetNonSandboxPath() );
 	// Copy the file into the sandbox
 	LowerLevel->CopyFile(*InPath.GetSandboxPath(), *InPath.GetNonSandboxPath());
 
