Fix wrong IsSet

#rnx

diff --git a/Engine/Plugins/Developer/Concert/ConcertSync/ConcertSyncClient/Source/ConcertSyncClient/Private/ConcertClientDesktopPresenceActor.cpp b/Engine/Plugins/Developer/Concert/ConcertSync/ConcertSyncClient/Source/ConcertSyncClient/Private/ConcertClientDesktopPresenceActor.cpp
--- a/Engine/Plugins/Developer/Concert/ConcertSync/ConcertSyncClient/Source/ConcertSyncClient/Private/ConcertClientDesktopPresenceActor.cpp
+++ b/Engine/Plugins/Developer/Concert/ConcertSync/ConcertSyncClient/Source/ConcertSyncClient/Private/ConcertClientDesktopPresenceActor.cpp
@@ -188,7 +188,7 @@ void AConcertClientDesktopPresenceActor::HideLaser()
 
 	if (LaserEndMovement.IsSet())
 	{
-		LaserEndMovement.IsSet();
+		LaserEndMovement.Reset();
 	}
 }
 
diff --git a/Engine/Plugins/Developer/Concert/ConcertSync/ConcertSyncClient/Source/ConcertSyncClient/Private/ConcertClientVRPresenceActor.cpp b/Engine/Plugins/Developer/Concert/ConcertSync/ConcertSyncClient/Source/ConcertSyncClient/Private/ConcertClientVRPresenceActor.cpp
--- a/Engine/Plugins/Developer/Concert/ConcertSync/ConcertSyncClient/Source/ConcertSyncClient/Private/ConcertClientVRPresenceActor.cpp
+++ b/Engine/Plugins/Developer/Concert/ConcertSync/ConcertSyncClient/Source/ConcertSyncClient/Private/ConcertClientVRPresenceActor.cpp
@@ -401,7 +401,7 @@ void AConcertClientVRPresenceActor::HideLeftLaser()
 
 	if (LeftLaserEndMovement.IsSet())
 	{
-		LeftLaserEndMovement.IsSet();
+		LeftLaserEndMovement.Reset();
 	}
 }
 
@@ -429,7 +429,7 @@ void AConcertClientVRPresenceActor::HideRightLaser()
 
 	if (RightLaserEndMovement.IsSet())
 	{
-		RightLaserEndMovement.IsSet();
+		RightLaserEndMovement.Reset();
 	}
 }
 
