Fixed Concert potentially calling its PostEdit functions with a stale ComponentToWorld

The PostEdit functions may read this data (eg, to get the world location of an actor), and it being stale could result in things reading the wrong location if a transform/attachment was modified as part of the Concert event.

#rnx

diff --git a/Engine/Plugins/Developer/Concert/ConcertSync/ConcertSyncClient/Source/ConcertSyncClient/Private/ConcertClientTransactionBridge.cpp b/Engine/Plugins/Developer/Concert/ConcertSync/ConcertSyncClient/Source/ConcertSyncClient/Private/ConcertClientTransactionBridge.cpp
--- a/Engine/Plugins/Developer/Concert/ConcertSync/ConcertSyncClient/Source/ConcertSyncClient/Private/ConcertClientTransactionBridge.cpp
+++ b/Engine/Plugins/Developer/Concert/ConcertSync/ConcertSyncClient/Source/ConcertSyncClient/Private/ConcertClientTransactionBridge.cpp
@@ -619,6 +619,13 @@ void ProcessTransactionEvent(const FConcertTransactionEventBase& InEvent, const
 		}
 
 #if WITH_EDITOR
+		// Ensure any scene components have an up-to-date ComponentToWorld before invoking their PostEdit logic, as this 
+		// update may have dirtied it, and PostEdit may try and read it (eg, to get the world position of an actor)
+		if (USceneComponent* TransactionSceneComponent = Cast<USceneComponent>(TransactionObject))
+		{
+			TransactionSceneComponent->UpdateComponentToWorld();
+		}
+
 		// We need to manually call OnObjectPropertyChanged as PostEditUndo calls the PostEditChange version that skips it, but we have things that rely on it being called
 		// For snapshot events this also triggers PostEditChange directly since we can skip the call to PostEditUndo
 		for (const FConcertSerializedPropertyData& PropertyData : ObjectUpdate.PropertyDatas)
