Fix skipping dirty sub-levels when exiting a session.  Any changes to the sub-level ended up getting persisted which may not be what the user wants.

diff --git a/Engine/Plugins/Developer/Concert/ConcertSync/ConcertSyncClient/Source/ConcertSyncClient/Private/ConcertSyncClientUtil.cpp b/Engine/Plugins/Developer/Concert/ConcertSync/ConcertSyncClient/Source/ConcertSyncClient/Private/ConcertSyncClientUtil.cpp
--- a/Engine/Plugins/Developer/Concert/ConcertSync/ConcertSyncClient/Source/ConcertSyncClient/Private/ConcertSyncClientUtil.cpp
+++ b/Engine/Plugins/Developer/Concert/ConcertSync/ConcertSyncClient/Source/ConcertSyncClient/Private/ConcertSyncClientUtil.cpp
@@ -679,7 +679,8 @@ void HotReloadPackages(TArrayView<const FName> InPackageNames)
 			{
 				bAddPersistentLevel = ShouldReloadPersistentLevel(ExistingPackage);
 			}
-			else
+
+			if (!ExistingPackage->ContainsMap() || ExistingPackage->IsDirty())
 			{
 				ExistingPackages.Add(ExistingPackage);
 			}
