MU Replication: Added UConcertSyncReplicationConfig which contains configurable actions to execute after certain properties are replicated.
By default, transform properties of USceneComponent will call UpdateComponentTransform and USceneComponent::bVisible will call MarkRenderStateDirty.

diff --git a/Engine/Plugins/Developer/Concert/ConcertSync/ConcertSyncClient/Source/ConcertSyncClient/Private/Replication/Manager/ReplicationManagerState.h b/Engine/Plugins/Developer/Concert/ConcertSync/ConcertSyncClient/Source/ConcertSyncClient/Private/Replication/Manager/ReplicationManagerState.h
--- a/Engine/Plugins/Developer/Concert/ConcertSync/ConcertSyncClient/Source/ConcertSyncClient/Private/Replication/Manager/ReplicationManagerState.h
+++ b/Engine/Plugins/Developer/Concert/ConcertSync/ConcertSyncClient/Source/ConcertSyncClient/Private/Replication/Manager/ReplicationManagerState.h
@@ -6,6 +6,8 @@
 #include "Replication/Manager/Utils/ReplicationManagerUtils.h"
 #include "Replication/Messages/Muting.h"
 
+class UConcertSyncReplicationConfig;
+
 namespace UE::ConcertSyncClient::Replication
 {
 	class FReplicationManager;
diff --git a/Engine/Plugins/Developer/Concert/ConcertSync/ConcertSyncClient/Source/ConcertSyncClient/Private/Replication/Processing/ObjectReplicationApplierProcessor.cpp b/Engine/Plugins/Developer/Concert/ConcertSync/ConcertSyncClient/Source/ConcertSyncClient/Private/Replication/Processing/ObjectReplicationApplierProcessor.cpp
--- a/Engine/Plugins/Developer/Concert/ConcertSync/ConcertSyncClient/Source/ConcertSyncClient/Private/Replication/Processing/ObjectReplicationApplierProcessor.cpp
+++ b/Engine/Plugins/Developer/Concert/ConcertSync/ConcertSyncClient/Source/ConcertSyncClient/Private/Replication/Processing/ObjectReplicationApplierProcessor.cpp
@@ -3,14 +3,21 @@
 #include "ObjectReplicationApplierProcessor.h"
 
 #include "ConcertLogGlobal.h"
-#include "Replication/IConcertClientReplicationBridge.h"
+#include "ConcertSyncSettings.h"
+#include "HAL/IConsoleManager.h"
 #include "Replication/Formats/IObjectReplicationFormat.h"
-
-#include "Components/SceneComponent.h"
+#include "Replication/IConcertClientReplicationBridge.h"
+#include "Replication/Processing/Actions/ReplicationActionDispatcher.h"
 #include "Trace/ConcertProtocolTrace.h"
 
 namespace UE::ConcertSyncClient::Replication
 {
+	static TAutoConsoleVariable<bool> CVarDebugReplicationActions(
+		TEXT("Concert.Replication.DebugReplicationActions"),
+		false,
+		TEXT("Whether to debug the actions taken by UConcertSyncReplicationConfig::PostReplicateActions")
+	);
+	
 	FObjectReplicationApplierProcessor::FObjectReplicationApplierProcessor(
 		IConcertClientReplicationBridge& ReplicationBridge,
 		ConcertSyncCore::IObjectReplicationFormat& ReplicationFormat,
@@ -31,27 +38,34 @@ namespace UE::ConcertSyncClient::Replication
 			return;
 		}
 
+		// Get the post replication actions
+		const TArray<FConcertReplicationActionEntry> DummyActions;
+		const TArray<FConcertReplicationActionEntry>& Actions = UConcertSyncConfig::Get()->PostReplicateActions;
+		const bool bDebugActions = CVarDebugReplicationActions.GetValueOnGameThread();
+		
 		bool bAppliedData = false;
-		GetDataSource().ExtractReplicationDataForObject(Args.ObjectInfo.ObjectId, [this, &Args, Object, &bAppliedData](const FConcertSessionSerializedPayload& Payload)
-		{
-			CONCERT_TRACE_REPLICATION_OBJECT_SCOPE(SerializeReceivedObject, Args.ObjectInfo.ObjectId.Object, Args.ObjectInfo.SequenceId);
-			bAppliedData = true;
-
-			// We're technically modifying the package so mark it modified. This will make Concert / Multi-User revert the changes on leaving the session.
-			// This could be made faster by caching the package. For now leave it and possibly profile in future if needed.
-			Object->MarkPackageDirty();
-			
-			// TODO DP UE-193659: This is very hacky and leaves performance on the table... this is in case ApplyReplicationEvent updates the transform
-			if (USceneComponent* SceneComponent = Cast<USceneComponent>(Object))
+		GetDataSource().ExtractReplicationDataForObject(Args.ObjectInfo.ObjectId,
+			[this, &Args, Object, Actions, bDebugActions, &bAppliedData](const FConcertSessionSerializedPayload& Payload)
 			{
-				ReplicationFormat.ApplyReplicationEvent(*Object, Payload);
-				SceneComponent->UpdateComponentToWorld();
-			}
-			else
-			{
-				ReplicationFormat.ApplyReplicationEvent(*Object, Payload);
-			}
-		});
+				CONCERT_TRACE_REPLICATION_OBJECT_SCOPE(SerializeReceivedObject, Args.ObjectInfo.ObjectId.Object, Args.ObjectInfo.SequenceId);
+				bAppliedData = true;
+
+				// We're technically modifying the package so mark it modified. This will make Concert / Multi-User revert the changes on leaving the session.
+				// This could be made faster by caching the package. For now leave it and possibly profile in future if needed.
+				Object->MarkPackageDirty();
+
+				ConcertSyncCore::FReplicationActionDispatcher ReplicationActionDispatcher(Actions, bDebugActions);
+				UE_CLOG(bDebugActions, LogConcert, Log, TEXT("===== %s actions ===="), *Object->GetPathName());
+				const auto ProcessReplicatedProperty = [&ReplicationActionDispatcher](const FArchiveSerializedPropertyChain& Chain, FProperty& Property)
+				{
+					ReplicationActionDispatcher.OnReplicateProperty(Property);
+				};
+				ReplicationFormat.ApplyReplicationEvent(*Object, Payload, ProcessReplicatedProperty);
+
+				// Execute any property dependent conditions, such as UpdateComponentToWorld for RelativeLocation, etc.
+				const ConcertSyncCore::FReplicationActionArgs PostReplicationArgs(Args.ObjectInfo.ObjectId, Object);
+				ReplicationActionDispatcher.ExecuteActions(PostReplicationArgs);
+			});
 		
 		// This should not happen. If it does, we're wasting network bandwidth.
 		UE_CLOG(!bAppliedData, LogConcert, Warning, TEXT("Replication: Server sent data that could not be applied (likely it was empty) for object %s from stream %s"), *Args.ObjectInfo.ObjectId.Object.ToString(), *Args.ObjectInfo.ObjectId.StreamId.ToString());
diff --git a/Engine/Plugins/Developer/Concert/ConcertSync/ConcertSyncClient/Source/ConcertSyncClient/Private/Replication/Processing/ObjectReplicationApplierProcessor.h b/Engine/Plugins/Developer/Concert/ConcertSync/ConcertSyncClient/Source/ConcertSyncClient/Private/Replication/Processing/ObjectReplicationApplierProcessor.h
--- a/Engine/Plugins/Developer/Concert/ConcertSync/ConcertSyncClient/Source/ConcertSyncClient/Private/Replication/Processing/ObjectReplicationApplierProcessor.h
+++ b/Engine/Plugins/Developer/Concert/ConcertSync/ConcertSyncClient/Source/ConcertSyncClient/Private/Replication/Processing/ObjectReplicationApplierProcessor.h
@@ -5,6 +5,7 @@
 #include "Replication/Processing/ObjectReplicationProcessor.h"
 #include "Replication/Processing/ReplicationDataQueuer.h"
 
+class UConcertSyncConfig;
 class IConcertClientReplicationBridge;
 
 namespace UE::ConcertSyncClient::Replication
@@ -14,7 +15,7 @@ namespace UE::ConcertSyncClient::Replication
 	{
 	public:
 
-		FObjectReplicationApplierProcessor(
+		explicit FObjectReplicationApplierProcessor(
 			IConcertClientReplicationBridge& ReplicationBridge UE_LIFETIMEBOUND,
 			ConcertSyncCore::IObjectReplicationFormat& ReplicationFormat UE_LIFETIMEBOUND,
 			ConcertSyncCore::IReplicationDataSource& DataSource UE_LIFETIMEBOUND
diff --git a/Engine/Plugins/Developer/Concert/ConcertSync/ConcertSyncCore/Config/BaseConcertSyncCore.ini b/Engine/Plugins/Developer/Concert/ConcertSync/ConcertSyncCore/Config/BaseConcertSyncCore.ini
--- a/Engine/Plugins/Developer/Concert/ConcertSync/ConcertSyncCore/Config/BaseConcertSyncCore.ini
+++ b/Engine/Plugins/Developer/Concert/ConcertSync/ConcertSyncCore/Config/BaseConcertSyncCore.ini
@@ -13,3 +13,6 @@
 +ExcludeTransactionClassFilters=(ObjectOuterClass=None,ObjectClasses=(Class'/Script/Landscape.LandscapeProxy',Class'/Script/Landscape.LandscapeHeightfieldCollisionComponent',Class'/Script/Landscape.LandscapeComponent'))
 ; If you wish to use Composure in Multi-user, you may add that line under [/Script/ConcertSyncCore.ConcertSyncConfig] in your project engine config settings.
 ;+IncludeObjectClassFilters=(ObjectOuterClass=None,ObjectClasses=(Class'/Script/Composure.CompositingElementPass',Class'/Script/Composure.ComposurePostProcessPassPolicy'))
+
++PostReplicateActions=(Action=/Script/ConcertSyncCore.ConcertReplicationAction_UpdateComponentToWorld(),Properties=("/Script/Engine.SceneComponent:RelativeLocation","/Script/Engine.SceneComponent:RelativeRotation","/Script/Engine.SceneComponent:RelativeScale3D","/Script/Engine.SceneComponent:bAbsoluteLocation","/Script/Engine.SceneComponent:bAbsoluteRotation","/Script/Engine.SceneComponent:bAbsoluteScale"))
++PostReplicateActions=(Action=/Script/ConcertSyncCore.ConcertReplicationAction_MarkRenderStateDirty(),Properties=("/Script/Engine.SceneComponent:bVisible"))
diff --git a/Engine/Plugins/Developer/Concert/ConcertSync/ConcertSyncCore/Source/ConcertSyncCore/Private/Replication/Formats/FullObjectFormat.cpp b/Engine/Plugins/Developer/Concert/ConcertSync/ConcertSyncCore/Source/ConcertSyncCore/Private/Replication/Formats/FullObjectFormat.cpp
--- a/Engine/Plugins/Developer/Concert/ConcertSync/ConcertSyncCore/Source/ConcertSyncCore/Private/Replication/Formats/FullObjectFormat.cpp
+++ b/Engine/Plugins/Developer/Concert/ConcertSync/ConcertSyncCore/Source/ConcertSyncCore/Private/Replication/Formats/FullObjectFormat.cpp
@@ -49,14 +49,35 @@ namespace UE::ConcertSyncCore
 
 		class FFulObjectFormatReader : public FObjectReader
 		{
+			const IObjectReplicationFormat::FOnPropertyVisitedFunc& OnPrePropertySerialized;
 		public:
 
-			FFulObjectFormatReader(TArray<uint8>& Data)
+			FFulObjectFormatReader(
+				TArray<uint8>& Data, const IObjectReplicationFormat::FOnPropertyVisitedFunc& OnPrePropertySerialized UE_LIFETIMEBOUND
+				)
 				: FObjectReader(Data)
+				, OnPrePropertySerialized(OnPrePropertySerialized)
 			{
 				constexpr bool bIsLoading = true;
 				ApplySharedArchiveConfig(*this, bIsLoading);
 			}
+
+			virtual void PushSerializedProperty(class FProperty* InProperty, const bool bIsEditorOnlyProperty) override
+			{
+				FObjectReader::PushSerializedProperty(InProperty, bIsEditorOnlyProperty);
+				
+				if (const FArchiveSerializedPropertyChain* PropertyChain = GetSerializedPropertyChain())
+				{
+					OnPrePropertySerialized(*PropertyChain, *InProperty);
+				}
+				else
+				{
+					// Value returned by GetSerializedPropertyChain is only set once serialization has gone into a non-root property...
+					// Callback takes const FArchiveSerializedPropertyChain& to simplify API user code.
+					FArchiveSerializedPropertyChain Dummy;
+					OnPrePropertySerialized(Dummy, *InProperty);
+				}
+			}
 		};
 	}
 	
@@ -81,13 +102,13 @@ namespace UE::ConcertSyncCore
 		Base = Newer; 
 	}
 
-	void FFullObjectFormat::ApplyReplicationEvent(UObject& Object, const FConcertSessionSerializedPayload& Payload)
+	void FFullObjectFormat::ApplyReplicationEvent(UObject& Object, const FConcertSessionSerializedPayload& Payload, const FOnPropertyVisitedFunc& OnPrePropertySerialized)
 	{
 		// This can be quite some data... maybe reading could be done without creating a copy...
 		FFullObjectReplicationData FullObjectFormat;
 		Payload.GetTypedPayload(FullObjectFormat);
 
-		FullObjectFormat::FFulObjectFormatReader WriterArchive(FullObjectFormat.SerializedObjectData.Bytes);
+		FullObjectFormat::FFulObjectFormatReader WriterArchive(FullObjectFormat.SerializedObjectData.Bytes, OnPrePropertySerialized);
 		constexpr bool bInLoadIfFindFails = false;
 		FObjectAndNameAsStringProxyArchive ObjectAndNameProxy(WriterArchive, bInLoadIfFindFails);
 		Object.Serialize(ObjectAndNameProxy);
diff --git a/Engine/Plugins/Developer/Concert/ConcertSync/ConcertSyncCore/Source/ConcertSyncCore/Private/Replication/Processing/Actions/ReplicationActionDispatcher.cpp b/Engine/Plugins/Developer/Concert/ConcertSync/ConcertSyncCore/Source/ConcertSyncCore/Private/Replication/Processing/Actions/ReplicationActionDispatcher.cpp
new file mode 100644
--- /dev/null
+++ b/Engine/Plugins/Developer/Concert/ConcertSync/ConcertSyncCore/Source/ConcertSyncCore/Private/Replication/Processing/Actions/ReplicationActionDispatcher.cpp
@@ -0,0 +1,58 @@
+﻿// Copyright Epic Games, Inc. All Rights Reserved.
+
+#include "Replication/Processing/Actions/ReplicationActionDispatcher.h"
+
+#include "ConcertLogGlobal.h"
+#include "Replication/Data/ReplicationActionEntry.h"
+
+namespace UE::ConcertSyncCore
+{
+FReplicationActionDispatcher::FReplicationActionDispatcher(const TArray<FConcertReplicationActionEntry>& InActions, bool bDebugActions)
+	: Actions(InActions)
+	, bDebugActions(bDebugActions)
+	, ActionsToTrigger(false, InActions.Num())
+{}
+
+void FReplicationActionDispatcher::OnReplicateProperty(const FProperty& InProperty)
+{
+	for (int32 Index = 0; Index < Actions.Num(); ++Index)
+	{
+		const FConcertReplicationActionEntry& Entry = Actions[Index];
+		const FConcertReplicationAction* Action = Actions[Index].Action.GetPtr<FConcertReplicationAction>();
+		if (!Action || ActionsToTrigger[Index])
+		{
+			continue;
+		}
+
+		const TArray<TFieldPath<FProperty>>& Properties = Entry.Properties;
+		const int32 PropertyIndex = Properties.IndexOfByPredicate([&InProperty](const TFieldPath<FProperty>& FieldPath)
+		{
+			const FProperty* ResolvedProperty = FieldPath.Get();
+			return ResolvedProperty == &InProperty;
+		});
+		const bool bContainsProperty = Properties.IsValidIndex(PropertyIndex);
+		if (bContainsProperty)
+		{
+			ActionsToTrigger[Index] = true;
+			UE_CLOG(bDebugActions, LogConcert, Log, TEXT("Property %s executes action %s"),
+				*Properties[PropertyIndex].ToString(),
+				*Entry.Action.GetScriptStruct()->GetName()
+				);
+		}
+	}
+}
+
+void FReplicationActionDispatcher::ExecuteActions(const FReplicationActionArgs& InArgs)
+{
+	check(Actions.Num() == ActionsToTrigger.Num()); // We assume config has not changed since construction.
+
+	for (int32 Index = 0; Index < Actions.Num(); ++Index)
+	{
+		const FConcertReplicationAction* Action = Actions[Index].Action.GetPtr<FConcertReplicationAction>();
+		if (ActionsToTrigger[Index] && Action)
+		{
+			Action->Apply(InArgs);
+		}
+	}
+}
+}
diff --git a/Engine/Plugins/Developer/Concert/ConcertSync/ConcertSyncCore/Source/ConcertSyncCore/Public/ConcertSyncSettings.h b/Engine/Plugins/Developer/Concert/ConcertSync/ConcertSyncCore/Source/ConcertSyncCore/Public/ConcertSyncSettings.h
--- a/Engine/Plugins/Developer/Concert/ConcertSync/ConcertSyncCore/Source/ConcertSyncCore/Public/ConcertSyncSettings.h
+++ b/Engine/Plugins/Developer/Concert/ConcertSync/ConcertSyncCore/Source/ConcertSyncCore/Public/ConcertSyncSettings.h
@@ -3,6 +3,8 @@
 #pragma once
 
 #include "UObject/FieldPath.h"
+#include "UObject/UObjectGlobals.h"
+#include "Replication/Data/ReplicationActionEntry.h"
 #include "ConcertSyncSettings.generated.h"
 
 USTRUCT()
@@ -67,6 +69,9 @@ public:
 		, SnapshotTransactionsPerSecond(30.0f)
 	{}
 
+	/** @return The global settings objects. */
+	static UConcertSyncConfig* Get() { return GetMutableDefault<UConcertSyncConfig>(); }
+
 	/**
 	 * Should we ask before hot-reloading changed packages?
 	 * If disabled we will clobber any local changes when reloading packages.
@@ -139,4 +144,8 @@ public:
 	 */
 	UPROPERTY(config, EditAnywhere, Category = "Package Settings")
 	TArray<FPackageClassFilter> ExcludePackageClassFilters;
+	
+	/** Defines actions that are to be executed after certain properties are replicated, such as calling PostEditChange, etc. */
+	UPROPERTY(Config, EditAnywhere, Category = "Actions")
+	TArray<FConcertReplicationActionEntry> PostReplicateActions;
 };
diff --git a/Engine/Plugins/Developer/Concert/ConcertSync/ConcertSyncCore/Source/ConcertSyncCore/Public/Replication/Data/ReplicationActionEntry.h b/Engine/Plugins/Developer/Concert/ConcertSync/ConcertSyncCore/Source/ConcertSyncCore/Public/Replication/Data/ReplicationActionEntry.h
new file mode 100644
--- /dev/null
+++ b/Engine/Plugins/Developer/Concert/ConcertSync/ConcertSyncCore/Source/ConcertSyncCore/Public/Replication/Data/ReplicationActionEntry.h
@@ -0,0 +1,23 @@
+﻿// Copyright Epic Games, Inc. All Rights Reserved.
+
+#pragma once
+
+#include "Replication/Processing/Actions/ConcertReplicationAction.h"
+#include "StructUtils/InstancedStruct.h"
+#include "UObject/FieldPath.h"
+#include "ReplicationActionEntry.generated.h"
+
+/** Defines an action that should be performed when a property is replicated. */
+USTRUCT()
+struct FConcertReplicationActionEntry
+{
+	GENERATED_BODY()
+
+	/** The action to perform when a property is matched. */
+	UPROPERTY(EditAnywhere, Category = "Replication")
+	TInstancedStruct<FConcertReplicationAction> Action;
+
+	/** If any of these properties is replicated, perform Action on the replicated object. */
+	UPROPERTY(EditAnywhere, Category = "Replication")
+	TArray<TFieldPath<FProperty>> Properties;
+};
\ No newline at end of file
diff --git a/Engine/Plugins/Developer/Concert/ConcertSync/ConcertSyncCore/Source/ConcertSyncCore/Public/Replication/Formats/FullObjectFormat.h b/Engine/Plugins/Developer/Concert/ConcertSync/ConcertSyncCore/Source/ConcertSyncCore/Public/Replication/Formats/FullObjectFormat.h
--- a/Engine/Plugins/Developer/Concert/ConcertSync/ConcertSyncCore/Source/ConcertSyncCore/Public/Replication/Formats/FullObjectFormat.h
+++ b/Engine/Plugins/Developer/Concert/ConcertSync/ConcertSyncCore/Source/ConcertSyncCore/Public/Replication/Formats/FullObjectFormat.h
@@ -25,7 +25,7 @@ namespace UE::ConcertSyncCore
 		virtual TOptional<FConcertSessionSerializedPayload> CreateReplicationEvent(UObject& Object, FAllowPropertyFunc IsPropertyAllowedFunc) override;
 		virtual void ClearInternalCache(TArrayView<UObject> ObjectsToClear) override { /* There is no cache: Full object format has no smart building relying on past events - it always serializes everything */ }
 		virtual void CombineReplicationEvents(FConcertSessionSerializedPayload& Base, const FConcertSessionSerializedPayload& Newer) override;
-		virtual void ApplyReplicationEvent(UObject& Object, const FConcertSessionSerializedPayload& Payload) override;
+		virtual void ApplyReplicationEvent(UObject& Object, const FConcertSessionSerializedPayload& Payload, const FOnPropertyVisitedFunc& OnPrePropertySerialized) override;
 		//~ End FFullObjectFormat Interface
 	};
 }
diff --git a/Engine/Plugins/Developer/Concert/ConcertSync/ConcertSyncCore/Source/ConcertSyncCore/Public/Replication/Formats/IObjectReplicationFormat.h b/Engine/Plugins/Developer/Concert/ConcertSync/ConcertSyncCore/Source/ConcertSyncCore/Public/Replication/Formats/IObjectReplicationFormat.h
--- a/Engine/Plugins/Developer/Concert/ConcertSync/ConcertSyncCore/Source/ConcertSyncCore/Public/Replication/Formats/IObjectReplicationFormat.h
+++ b/Engine/Plugins/Developer/Concert/ConcertSync/ConcertSyncCore/Source/ConcertSyncCore/Public/Replication/Formats/IObjectReplicationFormat.h
@@ -5,6 +5,7 @@
 #include "Containers/ArrayView.h"
 #include "Templates/Function.h"
 #include "Misc/Optional.h"
+#include "Serialization/ArchiveSerializedPropertyChain.h"
 
 class FProperty;
 class UObject;
@@ -18,7 +19,6 @@ namespace UE::ConcertSyncCore
 	 *
 	 * Example implementations:
 	 * - full-object format that serializes all replicated properties and never does any combining (inefficient).
-	 *   TODO: Full-object is for the MVP prototype - should we remove it later?
 	 * - partial-object format that keeps track of past properties and only sends changed replicated properties.
 	 */
 	class CONCERTSYNCCORE_API IObjectReplicationFormat
@@ -27,6 +27,7 @@ namespace UE::ConcertSyncCore
 
 		/** Chain can be nullptr or non-null and empty - both imply the Property is a root property. */
 		using FAllowPropertyFunc = TFunctionRef<bool(const FArchiveSerializedPropertyChain* Chain, const FProperty& Property)>;
+		using FOnPropertyVisitedFunc = TFunctionRef<void(const FArchiveSerializedPropertyChain&, FProperty&)>;
 
 		/**
 		 * Produces a new payload from Object.
@@ -59,8 +60,15 @@ namespace UE::ConcertSyncCore
 		 */
 		virtual void CombineReplicationEvents(FConcertSessionSerializedPayload& Base, const FConcertSessionSerializedPayload& Newer) = 0;
 
-		/** Applies Payload, which was created by CreateReplicationEvent, to Object. */
-		virtual void ApplyReplicationEvent(UObject& Object, const FConcertSessionSerializedPayload& Payload) = 0;
+		/**
+		 * Applies Payload to Object.
+		 * @param Object The object to apply to
+		 * @param Payload The payload to apply. Must have been created via CreateReplicationEvent.
+		 * @param OnPrePropertySerialized Invoked when a property is about to be serialized. Useful for determining whether e.g. render state needs to be dirtied, etc.
+		 */
+		virtual void ApplyReplicationEvent(UObject& Object, const FConcertSessionSerializedPayload& Payload,
+			const FOnPropertyVisitedFunc& OnPrePropertySerialized = [](const FArchiveSerializedPropertyChain&, FProperty&){}
+			) = 0;
 
 		virtual ~IObjectReplicationFormat() = default;
 	};
diff --git a/Engine/Plugins/Developer/Concert/ConcertSync/ConcertSyncCore/Source/ConcertSyncCore/Public/Replication/Processing/Actions/ConcertReplicationAction.h b/Engine/Plugins/Developer/Concert/ConcertSync/ConcertSyncCore/Source/ConcertSyncCore/Public/Replication/Processing/Actions/ConcertReplicationAction.h
new file mode 100644
--- /dev/null
+++ b/Engine/Plugins/Developer/Concert/ConcertSync/ConcertSyncCore/Source/ConcertSyncCore/Public/Replication/Processing/Actions/ConcertReplicationAction.h
@@ -0,0 +1,41 @@
+﻿// Copyright Epic Games, Inc. All Rights Reserved.
+
+#pragma once
+
+#include "ConcertReplicationAction.generated.h"
+
+struct FConcertReplicatedObjectId;
+
+namespace UE::ConcertSyncCore
+{
+/** Arguments for FConcertReplicationAction::Apply. */
+struct FReplicationActionArgs
+{
+	/** Replication info about the object. */
+	const FConcertReplicatedObjectId& ObjectId;
+	
+	/** The resolved object that was replicated. */
+	const TNotNull<UObject*> Object;
+
+	explicit FReplicationActionArgs(const FConcertReplicatedObjectId& InObjectId, TNotNull<UObject*> InObject)
+		: ObjectId(InObjectId)
+		, Object(InObject)
+	{}
+};
+}
+
+/**
+ * An action to be performed in relation to replication, such as after an object has been replicated.
+ * 
+ * Sub-structs implement Apply, which e.g. calls PostEditChange, MarkRenderStateDirty, or some other custom action.
+ * This can be used in conjunction with TInstancedStruct so these actions can be set up in .ini files,
+ * i.e. UPROPERTY(Config) TInstancedStruct<FConcertReplicationAction>.
+ */
+USTRUCT()
+struct CONCERTSYNCCORE_API FConcertReplicationAction
+{
+	GENERATED_BODY()
+
+	virtual ~FConcertReplicationAction() = default;
+	virtual void Apply(const UE::ConcertSyncCore::FReplicationActionArgs& InArgs) const { checkNoEntry(); }
+};
\ No newline at end of file
diff --git a/Engine/Plugins/Developer/Concert/ConcertSync/ConcertSyncCore/Source/ConcertSyncCore/Public/Replication/Processing/Actions/ConcertReplicationAction_MarkRenderStateDirty.h b/Engine/Plugins/Developer/Concert/ConcertSync/ConcertSyncCore/Source/ConcertSyncCore/Public/Replication/Processing/Actions/ConcertReplicationAction_MarkRenderStateDirty.h
new file mode 100644
--- /dev/null
+++ b/Engine/Plugins/Developer/Concert/ConcertSync/ConcertSyncCore/Source/ConcertSyncCore/Public/Replication/Processing/Actions/ConcertReplicationAction_MarkRenderStateDirty.h
@@ -0,0 +1,26 @@
+﻿// Copyright Epic Games, Inc. All Rights Reserved.
+
+#pragma once
+
+#include "ConcertReplicationAction.h"
+#if WITH_ENGINE
+#include "Components/SceneComponent.h"
+#endif
+#include "ConcertReplicationAction_MarkRenderStateDirty.generated.h"
+
+/** Calls MarkRenderStateDirty on the object if it is a USceneComponent. */
+USTRUCT(DisplayName = "Mark Render State Dirty")
+struct CONCERTSYNCCORE_API FConcertReplicationAction_MarkRenderStateDirty : public FConcertReplicationAction
+{
+	GENERATED_BODY()
+
+	virtual void Apply(const UE::ConcertSyncCore::FReplicationActionArgs& InArgs) const override
+	{
+#if WITH_ENGINE
+		if (USceneComponent* SceneComponent = Cast<USceneComponent>(InArgs.Object))
+		{
+			SceneComponent->MarkRenderStateDirty();
+		}
+#endif
+	}
+};
\ No newline at end of file
diff --git a/Engine/Plugins/Developer/Concert/ConcertSync/ConcertSyncCore/Source/ConcertSyncCore/Public/Replication/Processing/Actions/ConcertReplicationAction_PostEditChange.h b/Engine/Plugins/Developer/Concert/ConcertSync/ConcertSyncCore/Source/ConcertSyncCore/Public/Replication/Processing/Actions/ConcertReplicationAction_PostEditChange.h
new file mode 100644
--- /dev/null
+++ b/Engine/Plugins/Developer/Concert/ConcertSync/ConcertSyncCore/Source/ConcertSyncCore/Public/Replication/Processing/Actions/ConcertReplicationAction_PostEditChange.h
@@ -0,0 +1,21 @@
+﻿// Copyright Epic Games, Inc. All Rights Reserved.
+
+#pragma once
+
+#include "ConcertReplicationAction.h"
+#include "UObject/Object.h"
+#include "ConcertReplicationAction_PostEditChange.generated.h"
+
+/** Calls PostEditChange on the object. */
+USTRUCT(DisplayName = "Post Edit Change")
+struct CONCERTSYNCCORE_API FConcertReplicationAction_PostEditChange : public FConcertReplicationAction
+{
+	GENERATED_BODY()
+
+	virtual void Apply(const UE::ConcertSyncCore::FReplicationActionArgs& InArgs) const override
+	{
+#if WITH_EDITOR
+		InArgs.Object->PostEditChange();
+#endif
+	}
+};
\ No newline at end of file
diff --git a/Engine/Plugins/Developer/Concert/ConcertSync/ConcertSyncCore/Source/ConcertSyncCore/Public/Replication/Processing/Actions/ConcertReplicationAction_UpdateComponentToWorld.h b/Engine/Plugins/Developer/Concert/ConcertSync/ConcertSyncCore/Source/ConcertSyncCore/Public/Replication/Processing/Actions/ConcertReplicationAction_UpdateComponentToWorld.h
new file mode 100644
--- /dev/null
+++ b/Engine/Plugins/Developer/Concert/ConcertSync/ConcertSyncCore/Source/ConcertSyncCore/Public/Replication/Processing/Actions/ConcertReplicationAction_UpdateComponentToWorld.h
@@ -0,0 +1,26 @@
+﻿// Copyright Epic Games, Inc. All Rights Reserved.
+
+#pragma once
+
+#include "ConcertReplicationAction.h"
+#if WITH_ENGINE
+#include "Components/SceneComponent.h"
+#endif
+#include "ConcertReplicationAction_UpdateComponentToWorld.generated.h"
+
+/** Calls ConcertReplicationAction_UpdateComponentToWorld on the object if it is a USceneComponent. */
+USTRUCT(DisplayName = "Update Component To World")
+struct CONCERTSYNCCORE_API FConcertReplicationAction_UpdateComponentToWorld : public FConcertReplicationAction
+{
+	GENERATED_BODY()
+
+	virtual void Apply(const UE::ConcertSyncCore::FReplicationActionArgs& InArgs) const override
+	{
+#if WITH_ENGINE
+		if (USceneComponent* SceneComponent = Cast<USceneComponent>(InArgs.Object))
+		{
+			SceneComponent->UpdateComponentToWorld();
+		}
+#endif
+	}
+};
\ No newline at end of file
diff --git a/Engine/Plugins/Developer/Concert/ConcertSync/ConcertSyncCore/Source/ConcertSyncCore/Public/Replication/Processing/Actions/ReplicationActionDispatcher.h b/Engine/Plugins/Developer/Concert/ConcertSync/ConcertSyncCore/Source/ConcertSyncCore/Public/Replication/Processing/Actions/ReplicationActionDispatcher.h
new file mode 100644
--- /dev/null
+++ b/Engine/Plugins/Developer/Concert/ConcertSync/ConcertSyncCore/Source/ConcertSyncCore/Public/Replication/Processing/Actions/ReplicationActionDispatcher.h
@@ -0,0 +1,37 @@
+﻿// Copyright Epic Games, Inc. All Rights Reserved.
+
+#pragma once
+
+#include "Containers/BitArray.h"
+
+struct FConcertReplicationActionEntry;
+
+namespace UE::ConcertSyncCore
+{
+struct FReplicationActionArgs;
+
+/** Util that you can use when calling IObjectReplicationFormat::ApplyReplicationEvent to execute actions. */
+class CONCERTSYNCCORE_API FReplicationActionDispatcher
+{
+public:
+
+	explicit FReplicationActionDispatcher(const TArray<FConcertReplicationActionEntry>& InActions UE_LIFETIMEBOUND, bool bDebugActions = false);
+
+	/** Checks whether InProperty should trigger actions. */
+	void OnReplicateProperty(const FProperty& InProperty);
+
+	/** Call all properties have been processed using OnReplicateProperty. Triggers the actions that need to triggered.  */
+	void ExecuteActions(const FReplicationActionArgs& InArgs);
+
+private:
+
+	/** The actions to perform. */
+	const TArray<FConcertReplicationActionEntry>& Actions;
+
+	/** Whether the actions should be debugged. */
+	const bool bDebugActions;
+
+	/** Bit mask of equal length as the action array. Each index marks whether the action at the equivalent action array index should be performed. */
+	TBitArray<> ActionsToTrigger;
+};
+}
diff --git a/Engine/Plugins/Developer/Concert/ConcertSync/ConcertSyncTest/Source/ConcertSyncTest/Private/Replication/Components/ReplicationCache.spec.cpp b/Engine/Plugins/Developer/Concert/ConcertSync/ConcertSyncTest/Source/ConcertSyncTest/Private/Replication/Components/ReplicationCache.spec.cpp
--- a/Engine/Plugins/Developer/Concert/ConcertSync/ConcertSyncTest/Source/ConcertSyncTest/Private/Replication/Components/ReplicationCache.spec.cpp
+++ b/Engine/Plugins/Developer/Concert/ConcertSync/ConcertSyncTest/Source/ConcertSyncTest/Private/Replication/Components/ReplicationCache.spec.cpp
@@ -32,7 +32,10 @@ namespace UE::ConcertSyncTests
 			BaseStruct.Float += NewerStruct.Float;
 			Base.SetTypedPayload(BaseStruct);
 		}
-		virtual void ApplyReplicationEvent(UObject& Object, const FConcertSessionSerializedPayload& Payload) override { return NotMocked<void>(); }
+		virtual void ApplyReplicationEvent(UObject& Object, const FConcertSessionSerializedPayload& Payload, const FOnPropertyVisitedFunc& OnPrePropertySerialized) override
+		{
+			return NotMocked<void>();
+		}
 
 		static FConcertReplication_ObjectReplicationEvent CreateEvent(FSoftObjectPath Object, float Value)
 		{
