Add explicit methods to start and stop the transaction bridge so that we prevent the transaction handler from running when not in a multi-user session.

diff --git a/Engine/Plugins/Developer/Concert/ConcertSync/ConcertSyncClient/Source/ConcertSyncClient/Private/ConcertClientTransactionBridge.cpp b/Engine/Plugins/Developer/Concert/ConcertSync/ConcertSyncClient/Source/ConcertSyncClient/Private/ConcertClientTransactionBridge.cpp
--- a/Engine/Plugins/Developer/Concert/ConcertSync/ConcertSyncClient/Source/ConcertSyncClient/Private/ConcertClientTransactionBridge.cpp
+++ b/Engine/Plugins/Developer/Concert/ConcertSync/ConcertSyncClient/Source/ConcertSyncClient/Private/ConcertClientTransactionBridge.cpp
@@ -18,6 +18,8 @@
 #include "LevelInstance/LevelInstanceActor.h"
 #include "LevelInstance/LevelInstanceComponent.h"
 #include "Components/PrimitiveComponent.h"
+#include "UObject/SoftObjectPath.h"
+#include "UObject/WeakObjectPtr.h"
 
 #if WITH_EDITOR
 	#include "Editor.h"
@@ -42,17 +44,31 @@ static TAutoConsoleVariable<int32> CVarIgnoreTransactionIncludeFilter(TEXT("Conc
 static TAutoConsoleVariable<int32> CVarAlwaysSuspendBroadcastUndoRedo(TEXT("Concert.AlwaysSuspendBroadcastUndoRedo"), 0,
 																	  TEXT("Always suspend the undo/redo broadcast message. Can help with transactions that steal user focus."));
 
+UClass* LoadClassHelper(const FSoftClassPath& ClassPath)
+{
+	static TMap<FSoftClassPath, TWeakObjectPtr<UClass>> ClassCache;
+	if (TWeakObjectPtr<UClass>* ExistingClass = ClassCache.Find(ClassPath))
+	{
+		if (ExistingClass->IsStale())
+		{
+			*ExistingClass = ClassPath.TryLoadClass<UObject>();
+		}
+		return ExistingClass->Get();
+	}
+	return ClassCache.Add(ClassPath, ClassPath.TryLoadClass<UObject>()).Get();
+}
+	
 bool RunTransactionFilters(const TArray<FTransactionClassFilter>& InFilters, UObject* InObject)
 {
 	bool bMatchFilter = false;
 	for (const FTransactionClassFilter& TransactionFilter : InFilters)
 	{
-		UClass* TransactionOuterClass = TransactionFilter.ObjectOuterClass.TryLoadClass<UObject>();
+		UClass* TransactionOuterClass = LoadClassHelper(TransactionFilter.ObjectOuterClass);
 		if (!TransactionOuterClass || InObject->IsInA(TransactionOuterClass))
 		{
 			for (const FSoftClassPath& ObjectClass : TransactionFilter.ObjectClasses)
 			{
-				UClass* TransactionClass = ObjectClass.TryLoadClass<UObject>();
+				UClass* TransactionClass = LoadClassHelper(ObjectClass);
 				if (TransactionClass && InObject->IsA(TransactionClass))
 				{
 					bMatchFilter = true;
@@ -714,27 +730,12 @@ FConcertClientTransactionBridge::FConcertClientTransactionBridge()
 	, bIncludeNonPropertyObjectData(true)
 	, bIncludeAnnotationObjectChanges(GetDefault<UConcertSyncConfig>()->bIncludeAnnotationObjectChanges)
 {
-	ConditionalBindUnderlyingLocalTransactionEvents();
-
-	FCoreDelegates::OnFEngineLoopInitComplete.AddRaw(this, &FConcertClientTransactionBridge::OnEngineInitComplete);
 	FCoreDelegates::OnEndFrame.AddRaw(this, &FConcertClientTransactionBridge::OnEndFrame);
 }
 
 FConcertClientTransactionBridge::~FConcertClientTransactionBridge()
 {
-#if WITH_EDITOR
-	// Unregister Object Transaction events
-	if (GUnrealEd)
-	{
-		if (UTransBuffer* TransBuffer = Cast<UTransBuffer>(GUnrealEd->Trans))
-		{
-			TransBuffer->OnTransactionStateChanged().RemoveAll(this);
-		}
-	}
-	FCoreUObjectDelegates::OnObjectTransacted.RemoveAll(this);
-#endif
-
-	FCoreDelegates::OnFEngineLoopInitComplete.RemoveAll(this);
+	StopBridge();
 	FCoreDelegates::OnEndFrame.RemoveAll(this);
 }
 
@@ -1316,6 +1317,43 @@ void FConcertClientTransactionBridge::ConditionalBindUnderlyingLocalTransactionE
 #endif
 }
 
+void FConcertClientTransactionBridge::StartBridge()
+{
+	if (GIsRunning)
+	{
+		ConditionalBindUnderlyingLocalTransactionEvents();
+	}
+	else
+	{
+		// Wait until the engine init has been completed to bind to transaction events.
+		FCoreDelegates::OnFEngineLoopInitComplete.AddRaw(this, &FConcertClientTransactionBridge::OnEngineInitComplete);
+	}
+}
+
+
+void FConcertClientTransactionBridge::StopBridge()
+{
+	FCoreDelegates::OnFEngineLoopInitComplete.RemoveAll(this);
+
+	if (!bHasBoundUnderlyingLocalTransactionEvents)
+	{
+		return;
+	}
+	
+#if WITH_EDITOR
+	// Unregister Object Transaction events
+	if (GUnrealEd)
+	{
+		if (UTransBuffer* TransBuffer = Cast<UTransBuffer>(GUnrealEd->Trans))
+		{
+			TransBuffer->OnTransactionStateChanged().RemoveAll(this);
+		}
+	}
+	FCoreUObjectDelegates::OnObjectTransacted.RemoveAll(this);
+#endif
+	bHasBoundUnderlyingLocalTransactionEvents = false;
+}
+
 void FConcertClientTransactionBridge::OnEngineInitComplete()
 {
 	ConditionalBindUnderlyingLocalTransactionEvents();
diff --git a/Engine/Plugins/Developer/Concert/ConcertSync/ConcertSyncClient/Source/ConcertSyncClient/Private/ConcertClientTransactionBridge.h b/Engine/Plugins/Developer/Concert/ConcertSync/ConcertSyncClient/Source/ConcertSyncClient/Private/ConcertClientTransactionBridge.h
--- a/Engine/Plugins/Developer/Concert/ConcertSync/ConcertSyncClient/Source/ConcertSyncClient/Private/ConcertClientTransactionBridge.h
+++ b/Engine/Plugins/Developer/Concert/ConcertSync/ConcertSyncClient/Source/ConcertSyncClient/Private/ConcertClientTransactionBridge.h
@@ -11,7 +11,7 @@ namespace ConcertClientTransactionBridgeUtil
 	struct FTransactedObjectState;
 }
 
-class FConcertClientTransactionBridge : public IConcertClientTransactionBridge
+class FConcertClientTransactionBridge final : public IConcertClientTransactionBridge
 {
 public:
 	FConcertClientTransactionBridge();
@@ -33,7 +33,10 @@ public:
 
 	virtual void RegisterTransactionFilter(FName FilterName, FOnFilterTransactionDelegate FilterDelegate) override;
 	virtual void UnregisterTransactionFilter(FName FilterName) override;
-	
+
+	virtual void StartBridge() override;
+	virtual void StopBridge() override;
+
 private:
 	/** Called to handle a transaction state change */
 	void HandleTransactionStateChanged(const FTransactionContext& InTransactionContext, const ETransactionStateEventType InTransactionState);
diff --git a/Engine/Plugins/Developer/Concert/ConcertSync/ConcertSyncClient/Source/ConcertSyncClient/Private/ConcertSyncClient.cpp b/Engine/Plugins/Developer/Concert/ConcertSync/ConcertSyncClient/Source/ConcertSyncClient/Private/ConcertSyncClient.cpp
--- a/Engine/Plugins/Developer/Concert/ConcertSync/ConcertSyncClient/Source/ConcertSyncClient/Private/ConcertSyncClient.cpp
+++ b/Engine/Plugins/Developer/Concert/ConcertSync/ConcertSyncClient/Source/ConcertSyncClient/Private/ConcertSyncClient.cpp
@@ -48,6 +48,7 @@ void FConcertSyncClient::Startup(const UConcertClientConfig* InClientConfig, con
 	// Boot the client instance
 	ConcertClient->Configure(InClientConfig);
 	ConcertClient->Startup();
+	Bridges.TransactionBridge->StartBridge();
 
 	// if auto connection, start auto-connection routine
 	if (InClientConfig->bAutoConnect && ConcertClient->CanAutoConnect())
@@ -58,6 +59,7 @@ void FConcertSyncClient::Startup(const UConcertClientConfig* InClientConfig, con
 
 void FConcertSyncClient::Shutdown()
 {
+	Bridges.TransactionBridge->StopBridge();
 	ConcertClient->Shutdown();
 }
 
diff --git a/Engine/Plugins/Developer/Concert/ConcertSync/ConcertSyncClient/Source/ConcertSyncClient/Public/IConcertClientTransactionBridge.h b/Engine/Plugins/Developer/Concert/ConcertSync/ConcertSyncClient/Source/ConcertSyncClient/Public/IConcertClientTransactionBridge.h
--- a/Engine/Plugins/Developer/Concert/ConcertSync/ConcertSyncClient/Source/ConcertSyncClient/Public/IConcertClientTransactionBridge.h
+++ b/Engine/Plugins/Developer/Concert/ConcertSync/ConcertSyncClient/Source/ConcertSyncClient/Public/IConcertClientTransactionBridge.h
@@ -201,6 +201,20 @@ public:
 
 	/** Callback to register delegate for handling transaction events */
 	virtual void UnregisterTransactionFilter(FName FilterName) = 0;
+
+	/**
+	 * Bind to the transaction system to receive transaction events.  In past versions of IConcertClientTransactionBridge,
+	 * this binding happened automatically.  To prevent possible performance impact from handling editor transactions when the bridge is
+	 * not needed, users now have to indicate when they want to start the bridge which binds to the underlying transaction system.
+	 */
+	virtual void StartBridge() = 0;
+
+	/**
+	 * Stop briding the transaction system.  After this call the bridge will unbind from any handlers to the underlying transaction system and
+	 * transaction events will not be processed. 
+	 */
+	virtual void StopBridge() = 0;
+
 protected:
 	/**
 	 * Function to access the internal bool controlling whether local transactions are currently being tracked.
