Fix horde issue 862319. Mark methods intended to be final instead of the entire class.

diff --git a/Engine/Plugins/Developer/Concert/ConcertSync/ConcertSyncClient/Source/ConcertSyncClient/Private/ConcertClientTransactionBridge.h b/Engine/Plugins/Developer/Concert/ConcertSync/ConcertSyncClient/Source/ConcertSyncClient/Private/ConcertClientTransactionBridge.h
--- a/Engine/Plugins/Developer/Concert/ConcertSync/ConcertSyncClient/Source/ConcertSyncClient/Private/ConcertClientTransactionBridge.h
+++ b/Engine/Plugins/Developer/Concert/ConcertSync/ConcertSyncClient/Source/ConcertSyncClient/Private/ConcertClientTransactionBridge.h
@@ -11,7 +11,7 @@ namespace ConcertClientTransactionBridgeUtil
 	struct FTransactedObjectState;
 }
 
-class FConcertClientTransactionBridge final : public IConcertClientTransactionBridge
+class FConcertClientTransactionBridge : public IConcertClientTransactionBridge
 {
 public:
 	FConcertClientTransactionBridge();
@@ -34,8 +34,8 @@ public:
 	virtual void RegisterTransactionFilter(FName FilterName, FOnFilterTransactionDelegate FilterDelegate) override;
 	virtual void UnregisterTransactionFilter(FName FilterName) override;
 
-	virtual void StartBridge() override;
-	virtual void StopBridge() override;
+	virtual void StartBridge() final;
+	virtual void StopBridge() final;
 
 private:
 	/** Called to handle a transaction state change */
