It use to be that transaction filters would check for the UMovieSceneSignedObject but there seems to be an increasingly number of types on sequencer that do not have this as a base. This change adds a custom filter to include anything that has an outermost of level sequence.  Added a CVar as a way to turn it off.

#rnx

diff --git a/Engine/Plugins/Developer/Concert/ConcertSync/ConcertSyncClient/Source/ConcertSyncClient/Private/ConcertClientSequencerManager.cpp b/Engine/Plugins/Developer/Concert/ConcertSync/ConcertSyncClient/Source/ConcertSyncClient/Private/ConcertClientSequencerManager.cpp
--- a/Engine/Plugins/Developer/Concert/ConcertSync/ConcertSyncClient/Source/ConcertSyncClient/Private/ConcertClientSequencerManager.cpp
+++ b/Engine/Plugins/Developer/Concert/ConcertSync/ConcertSyncClient/Source/ConcertSyncClient/Private/ConcertClientSequencerManager.cpp
@@ -5,6 +5,7 @@
 #include "ConcertTransportMessages.h"
 #include "Delegates/IDelegateInstance.h"
 #include "HAL/IConsoleManager.h"
+#include "LevelSequence.h"
 #include "Misc/AssertionMacros.h"
 #include "Misc/QualifiedFrameTime.h"
 #include "Modules/ModuleManager.h"
@@ -13,11 +14,14 @@
 
 #include "IConcertClient.h"
 #include "IConcertSyncClient.h"
+#include "IConcertSyncClientModule.h"
+
 #include "IConcertSession.h"
 #include "ConcertSettings.h"
 #include "ConcertClientSettings.h"
 #include "ConcertSyncArchives.h"
 #include "ConcertClientWorkspace.h"
+#include "ConcertTransactionEvents.h"
 
 #include "Engine/GameEngine.h"
 #include "MovieSceneSequence.h"
@@ -41,6 +45,9 @@ DEFINE_LOG_CATEGORY_STATIC(LogConcertSequencerSync, Warning, All)
 
 #if WITH_EDITOR
 
+// Allow transactions on all edits to level sequences.
+static TAutoConsoleVariable<bool> CVarAllowAllTransactionsOnLevelSequences(TEXT("Concert.AllowAllTransactionsOnLevelSequences"), true, TEXT("Allow all transactions to pass the transaction filter when those transactions have an outermost object that is a level sequence."));
+
 // Enable Sequence Playback Syncing
 static TAutoConsoleVariable<int32> CVarEnablePlaybackSync(TEXT("Concert.EnableSequencerPlaybackSync"), 1, TEXT("Enable Concert Sequencer Playback Syncing of opened Sequencer."));
 
@@ -218,7 +225,9 @@ FConcertClientSequencerManager::FConcertClientSequencerManager(IConcertSyncClien
 	if (GIsEditor)
 	{
 		ISequencerModule& SequencerModule = FModuleManager::Get().LoadModuleChecked<ISequencerModule>("Sequencer");
-		OnSequencerCreatedHandle = SequencerModule.RegisterOnSequencerCreated(FOnSequencerCreated::FDelegate::CreateRaw(this, &FConcertClientSequencerManager::OnSequencerCreated));
+		OnSequencerCreatedHandle = SequencerModule.RegisterOnSequencerCreated(
+			FOnSequencerCreated::FDelegate::CreateRaw(this, &FConcertClientSequencerManager::OnSequencerCreated)
+		);
 	}
 
 	FCoreUObjectDelegates::OnPackageReloaded.AddRaw(this, &FConcertClientSequencerManager::HandleAssetReload);
@@ -372,6 +381,21 @@ float FConcertClientSequencerManager::GetLatencyCompensationMs() const
 		: 0.0f;
 }
 
+ETransactionFilterResult FConcertClientSequencerManager::ShouldObjectBeTransacted(const FConcertTransactionFilterArgs& FilterArgs)
+{
+	const bool bShouldCheckOutermost = CVarAllowAllTransactionsOnLevelSequences.GetValueOnAnyThread();
+	if (bShouldCheckOutermost)
+	{
+		UObject* Outermost = IsValid(FilterArgs.ObjectToFilter) ? FilterArgs.ObjectToFilter->GetOutermostObject() : nullptr;
+		if (WeakSession.IsValid() && Outermost && Outermost->IsA<ULevelSequence>())
+		{
+			return ETransactionFilterResult::IncludeObject;
+		}
+	}
+
+	return ETransactionFilterResult::UseDefault;
+}
+
 void FConcertClientSequencerManager::Register(TSharedRef<IConcertClientSession> InSession)
 {
 	// Hold onto the session so we can trigger events
@@ -384,6 +408,14 @@ void FConcertClientSequencerManager::Register(TSharedRef<IConcertClientSession>
 	InSession->RegisterCustomEventHandler<FConcertSequencerStateSyncEvent>(this, &FConcertClientSequencerManager::OnStateSyncEvent);
 	InSession->RegisterCustomEventHandler<FConcertSequencerTimeAdjustmentEvent>(this, &FConcertClientSequencerManager::OnTimeAdjustmentEvent);
 
+	if (GIsEditor)
+	{
+		if (IConcertClientTransactionBridge* TransactionBridge = OwnerSyncClient->GetTransactionBridge())
+		{
+			TransactionBridge->RegisterTransactionFilter(TEXT("ConcertSequencerLS"), FOnFilterTransactionDelegate::CreateRaw(this, &FConcertClientSequencerManager::ShouldObjectBeTransacted));				
+		}
+	}
+
 	Preloader->OnRegister(InSession);
 }
 
@@ -403,6 +435,11 @@ void FConcertClientSequencerManager::Unregister(TSharedRef<IConcertClientSession
 
 		if (GEditor)
 		{
+			if (IConcertClientTransactionBridge* TransactionBridge = OwnerSyncClient->GetTransactionBridge())
+			{
+				TransactionBridge->UnregisterTransactionFilter(TEXT("ConcertSequencerLS"));
+			}
+			
 			TArray<UMovieSceneSequence*> Sequences;
 			for (FOpenSequencerData& OpenSequencer : OpenSequencers)
 			{
diff --git a/Engine/Plugins/Developer/Concert/ConcertSync/ConcertSyncClient/Source/ConcertSyncClient/Private/ConcertClientSequencerManager.h b/Engine/Plugins/Developer/Concert/ConcertSync/ConcertSyncClient/Source/ConcertSyncClient/Private/ConcertClientSequencerManager.h
--- a/Engine/Plugins/Developer/Concert/ConcertSync/ConcertSyncClient/Source/ConcertSyncClient/Private/ConcertClientSequencerManager.h
+++ b/Engine/Plugins/Developer/Concert/ConcertSync/ConcertSyncClient/Source/ConcertSyncClient/Private/ConcertClientSequencerManager.h
@@ -23,6 +23,8 @@ class FConcertClientWorkspace;
 struct FSequencerInitParams;
 class ULevelSequence;
 class FConcertClientSequencerStateEventPacer;
+enum class ETransactionFilterResult : uint8;
+struct FConcertTransactionFilterArgs;
 
 /**
  * Sequencer manager that is held by the client sync module that keeps track of open sequencer UIs, regardless of whether a session is open or not
@@ -355,6 +357,14 @@ private:
 	}
 
 private:
+	/**
+	 * Delegate to provide a hint to the transaction bridge that the given object should be transacted.  We special case LevelSequences
+	 * and allow any transaction that has a level sequence as an outermost object to be included because there are many object types
+	 * in level sequence transactions that do not share UMovieSceneSignedObject as a base class, which is used by the default transaction
+	 * filters.
+	 */
+	ETransactionFilterResult ShouldObjectBeTransacted(const FConcertTransactionFilterArgs& FilterArgs);
+		
 	/** Destroy the given sequence player with corresponding actor. */
 	void DestroyPlayer(ALevelSequenceActor* LevelSequenceActor);
 
