Fix values not being tracked by overridable serialization when applying a concert transaction

diff --git a/Engine/Plugins/Developer/Concert/ConcertSync/ConcertSyncClient/Source/ConcertSyncClient/Private/ConcertClientTransactionBridge.cpp b/Engine/Plugins/Developer/Concert/ConcertSync/ConcertSyncClient/Source/ConcertSyncClient/Private/ConcertClientTransactionBridge.cpp
--- a/Engine/Plugins/Developer/Concert/ConcertSync/ConcertSyncClient/Source/ConcertSyncClient/Private/ConcertClientTransactionBridge.cpp
+++ b/Engine/Plugins/Developer/Concert/ConcertSync/ConcertSyncClient/Source/ConcertSyncClient/Private/ConcertClientTransactionBridge.cpp
@@ -18,6 +18,7 @@
 #include "LevelInstance/LevelInstanceActor.h"
 #include "LevelInstance/LevelInstanceComponent.h"
 #include "Components/PrimitiveComponent.h"
+#include "UObject/OverridableManager.h"
 #include "UObject/SoftObjectPath.h"
 #include "UObject/WeakObjectPtr.h"
 
@@ -565,6 +566,7 @@ void ProcessTransactionEvent(const FConcertTransactionEventBase& InEvent, const
 				FEditPropertyChain PropertyChain;
 				PropertyChain.AddHead(TransactionProp);
 				FCoreUObjectDelegates::OnPreObjectPropertyChanged.Broadcast(TransactionObject, PropertyChain);
+				FOverridableManager::Get().PreOverrideProperty(TransactionObject, PropertyChain);
 			}
 		}
 	}
@@ -657,6 +659,9 @@ void ProcessTransactionEvent(const FConcertTransactionEventBase& InEvent, const
 
 				FPropertyChangedEvent PropertyChangedEvent(TransactionProp, bIsSnapshot ? EPropertyChangeType::Interactive : EPropertyChangeType::Unspecified);
 				FCoreUObjectDelegates::OnObjectPropertyChanged.Broadcast(TransactionObject, PropertyChangedEvent);
+				FEditPropertyChain PropertyChain;
+				PropertyChain.AddHead(TransactionProp);
+				FOverridableManager::Get().PostOverrideProperty(TransactionObject, PropertyChangedEvent, PropertyChain);
 			}
 		}
 
