Reset the FConcertClientSequencerStateEventPacer when we move from play to stop so that we properly send an event on restart.

#rnx

diff --git a/Engine/Plugins/Developer/Concert/ConcertSync/ConcertSyncClient/Source/ConcertSyncClient/Private/ConcertClientSequencerManager.cpp b/Engine/Plugins/Developer/Concert/ConcertSync/ConcertSyncClient/Source/ConcertSyncClient/Private/ConcertClientSequencerManager.cpp
--- a/Engine/Plugins/Developer/Concert/ConcertSync/ConcertSyncClient/Source/ConcertSyncClient/Private/ConcertClientSequencerManager.cpp
+++ b/Engine/Plugins/Developer/Concert/ConcertSync/ConcertSyncClient/Source/ConcertSyncClient/Private/ConcertClientSequencerManager.cpp
@@ -97,6 +97,12 @@ public:
 		return true;
 	}
 
+	void Reset()
+	{
+		bIsFirstEvent = true;
+		CurrentStateEvent.Reset();
+	}
+
 	void AddStateEvent(TSharedPtr<IConcertClientSession>& Session, FConcertSequencerStateEvent InEvent, FConcertClientSequencerManager::EPlaybackMode Mode)
 	{
 		if (!ShouldIncludeMessageForSend(Mode))
@@ -791,6 +797,7 @@ void FConcertClientSequencerManager::OnSequencerTimeChanged(TWeakPtr<ISequencer>
 			}
 			else
 			{
+				OpenSequencer->StateEventPacer->Reset();
 				OpenSequencer->PlaybackMode = EPlaybackMode::Undefined;
 			}
 		}
