Fixed nested entity creation in Concert

* UConcertClientEntityFactory now adds an entity to their outer entity, rather than always adding them at the top-level.
* UConcertClientEntityComponentFactory is now guaranteed to call InitializeObjects before UConcertClientEntityFactory, as components must initialize before entities.
* Cleaned up some redundant FName -> FString -> FName churn.

#rnx

diff --git a/Engine/Plugins/Developer/Concert/ConcertSync/ConcertSyncClient/Source/ConcertSyncClient/Private/ConcertClientTransactionBridge.cpp b/Engine/Plugins/Developer/Concert/ConcertSync/ConcertSyncClient/Source/ConcertSyncClient/Private/ConcertClientTransactionBridge.cpp
--- a/Engine/Plugins/Developer/Concert/ConcertSync/ConcertSyncClient/Source/ConcertSyncClient/Private/ConcertClientTransactionBridge.cpp
+++ b/Engine/Plugins/Developer/Concert/ConcertSync/ConcertSyncClient/Source/ConcertSyncClient/Private/ConcertClientTransactionBridge.cpp
@@ -486,9 +486,22 @@ void ProcessTransactionEvent(const FConcertTransactionEventBase& InEvent, const
 		}
 
 		// Run any deferred object factory initialization
-		for (const TTuple<const UConcertClientObjectFactory*, TArray<UObject*>>& ObjectsPendingInitializationPair : ObjectsPendingInitialization)
+		if (ObjectsPendingInitialization.Num() > 0)
 		{
-			ObjectsPendingInitializationPair.Key->InitializeObjects(ObjectsPendingInitializationPair.Value);
+			// If factories have initialization order dependencies, then make sure and respect those
+			TArray<const UConcertClientObjectFactory*, TInlineAllocator<4>> SortedFactories;
+			ObjectsPendingInitialization.GetKeys(SortedFactories);
+			Algo::Sort(SortedFactories,
+				[](const UConcertClientObjectFactory* ThisFactory, const UConcertClientObjectFactory* OtherFactory)
+				{
+					return ThisFactory->ShouldInitializeObjectsBefore(OtherFactory);
+				});
+
+			for (const UConcertClientObjectFactory* Factory : SortedFactories)
+			{
+				const TArray<UObject*>& Objects = ObjectsPendingInitialization.FindChecked(Factory);
+				Factory->InitializeObjects(Objects);
+			}
 		}
 	}
 
diff --git a/Engine/Plugins/Developer/Concert/ConcertSync/ConcertSyncClient/Source/ConcertSyncClient/Private/ConcertSyncClientUtil.cpp b/Engine/Plugins/Developer/Concert/ConcertSync/ConcertSyncClient/Source/ConcertSyncClient/Private/ConcertSyncClientUtil.cpp
--- a/Engine/Plugins/Developer/Concert/ConcertSync/ConcertSyncClient/Source/ConcertSyncClient/Private/ConcertSyncClientUtil.cpp
+++ b/Engine/Plugins/Developer/Concert/ConcertSync/ConcertSyncClient/Source/ConcertSyncClient/Private/ConcertSyncClientUtil.cpp
@@ -372,11 +372,11 @@ FGetObjectResult GetObject(const FConcertObjectId& InObjectId, const FName InNew
 			{
 				if (SourceObject)
 				{
-					bFactoryHandledCreation = Factory->DuplicateObject(ObjectResult.Obj, SourceObject, NewObjectOuter, ObjectClass, *ObjectNameToCreate.ToString(), (EObjectFlags)InObjectId.ObjectPersistentFlags);
+					bFactoryHandledCreation = Factory->DuplicateObject(ObjectResult.Obj, SourceObject, NewObjectOuter, ObjectClass, ObjectNameToCreate, (EObjectFlags)InObjectId.ObjectPersistentFlags);
 				}
 				else
 				{
-					bFactoryHandledCreation = Factory->CreateObject(ObjectResult.Obj, NewObjectOuter, ObjectClass, *ObjectNameToCreate.ToString(), (EObjectFlags)InObjectId.ObjectPersistentFlags);
+					bFactoryHandledCreation = Factory->CreateObject(ObjectResult.Obj, NewObjectOuter, ObjectClass, ObjectNameToCreate, (EObjectFlags)InObjectId.ObjectPersistentFlags);
 				}
 			}
 			if (!bFactoryHandledCreation)
@@ -456,7 +456,7 @@ FGetObjectResult GetObject(const FConcertObjectId& InObjectId, const FName InNew
 				{
 					if (SourceObject)
 					{
-						ObjectResult.Obj = DuplicateObject<UObject>(SourceObject, NewObjectOuter, *ObjectNameToCreate.ToString());
+						ObjectResult.Obj = DuplicateObject<UObject>(SourceObject, NewObjectOuter, ObjectNameToCreate);
 						if (ObjectResult.Obj)
 						{
 							ObjectResult.Obj->SetFlags((EObjectFlags)InObjectId.ObjectPersistentFlags);
@@ -464,7 +464,7 @@ FGetObjectResult GetObject(const FConcertObjectId& InObjectId, const FName InNew
 					}
 					else
 					{
-						ObjectResult.Obj = NewObject<UObject>(NewObjectOuter, ObjectClass, *ObjectNameToCreate.ToString(), (EObjectFlags)InObjectId.ObjectPersistentFlags);
+						ObjectResult.Obj = NewObject<UObject>(NewObjectOuter, ObjectClass, ObjectNameToCreate, (EObjectFlags)InObjectId.ObjectPersistentFlags);
 					}
 
 					if (UActorComponent* NewComponent = Cast<UActorComponent>(ObjectResult.Obj))
diff --git a/Engine/Plugins/Developer/Concert/ConcertSync/ConcertSyncClient/Source/ConcertSyncClient/Public/ConcertClientObjectFactory.h b/Engine/Plugins/Developer/Concert/ConcertSync/ConcertSyncClient/Source/ConcertSyncClient/Public/ConcertClientObjectFactory.h
--- a/Engine/Plugins/Developer/Concert/ConcertSync/ConcertSyncClient/Source/ConcertSyncClient/Public/ConcertClientObjectFactory.h
+++ b/Engine/Plugins/Developer/Concert/ConcertSync/ConcertSyncClient/Source/ConcertSyncClient/Public/ConcertClientObjectFactory.h
@@ -62,6 +62,15 @@ public:
 		return false;
 	}
 
+	/**
+	 * Should this factory run InitializeObjects before OtherFactory?
+	 * @return True if this factory needs to run its InitializeObjects before OtherFactory (to control initialization order between factories), or false if it doesn't matter which order they run.
+	 */
+	virtual bool ShouldInitializeObjectsBefore(const UConcertClientObjectFactory* OtherFactory) const
+	{
+		return false;
+	}
+
 	/**
 	 * Perform any additional initialization on the given objects.
 	 * @note This is called after CreateObject/DuplicateObject has been called for every object in the transaction, and can be used for delayed initialization of interdependent objects.
