Prevent flush loaders from running during shutdown

diff --git a/Engine/Plugins/Developer/Concert/ConcertSync/ConcertSyncClient/Source/ConcertSyncClient/Private/ConcertSyncClientUtil.cpp b/Engine/Plugins/Developer/Concert/ConcertSync/ConcertSyncClient/Source/ConcertSyncClient/Private/ConcertSyncClientUtil.cpp
--- a/Engine/Plugins/Developer/Concert/ConcertSync/ConcertSyncClient/Source/ConcertSyncClient/Private/ConcertSyncClientUtil.cpp
+++ b/Engine/Plugins/Developer/Concert/ConcertSync/ConcertSyncClient/Source/ConcertSyncClient/Private/ConcertSyncClientUtil.cpp
@@ -544,7 +544,6 @@ void SerializeObject(FConcertLocalIdentifierTable* InLocalIdentifierTable, const
 	ObjectWriter.SerializeObject(InObject, InProperties);
 }
 
-
 void FlushPackageLoading(const FName InPackageName)
 {
 	FlushPackageLoading(InPackageName.ToString());
@@ -552,6 +551,11 @@ void FlushPackageLoading(const FName InPackageName)
 
 void FlushPackageLoading(const FString& InPackageName, bool bForceBulkDataLoad)
 {
+	if (IsEngineExitRequested())
+	{
+		return;
+	}
+
 	UPackage* ExistingPackage = FindPackage(nullptr, *InPackageName);
 	if (ExistingPackage)
 	{
