Skip reloading detaching bulk data when persisting files back to the content tree.   The bulk data will be different because we will be reading the file from the non-sandbox area.

diff --git a/Engine/Plugins/Developer/Concert/ConcertSync/ConcertSyncClient/Source/ConcertSyncClient/Private/ConcertSandboxPlatformFile.cpp b/Engine/Plugins/Developer/Concert/ConcertSync/ConcertSyncClient/Source/ConcertSyncClient/Private/ConcertSandboxPlatformFile.cpp
--- a/Engine/Plugins/Developer/Concert/ConcertSync/ConcertSyncClient/Source/ConcertSyncClient/Private/ConcertSandboxPlatformFile.cpp
+++ b/Engine/Plugins/Developer/Concert/ConcertSync/ConcertSyncClient/Source/ConcertSyncClient/Private/ConcertSandboxPlatformFile.cpp
@@ -23,12 +23,12 @@
 namespace ConcertSandboxPlatformFileUtil
 {
 
-bool FlushPackageFile(const FString& InFilename, FName* OutPackageName = nullptr, bool bForceLoad = true)
+bool FlushPackageFile(const FString& InFilename, FName* OutPackageName = nullptr, bool bForceBulkDataLoad = true)
 {
 	FString PackageName;
 	if (FPackageName::TryConvertFilenameToLongPackageName(InFilename, PackageName))
 	{
-		ConcertSyncClientUtil::FlushPackageLoading(PackageName, bForceLoad);
+		ConcertSyncClientUtil::FlushPackageLoading(PackageName, bForceBulkDataLoad);
 		if (OutPackageName)
 		{
 			*OutPackageName = *PackageName;
@@ -677,7 +677,7 @@ bool FConcertSandboxPlatformFile::CopyFileWithSCC(
 	FPersistResult& OutResult)
 {
 	// If this file maps to a package then we need to flush its linker so that we can overwrite the file on disk
-	ConcertSandboxPlatformFileUtil::FlushPackageFile(InToFilename);
+	ConcertSandboxPlatformFileUtil::FlushPackageFile(InToFilename, nullptr, false);
 
 	ISourceControlProvider* SourceControlProvider = InParam.SourceControlProvider;
 	// Get the source control state of the destination file
@@ -767,8 +767,8 @@ bool FConcertSandboxPlatformFile::CopyFileWithSCC(
 /** Delete file using source control */
 bool FConcertSandboxPlatformFile::DeleteFileWithSCC(const TCHAR* InFilename, const FPersistParameters& InParam, FPersistResult& OutResult)
 {
-	// If this file maps to a package then we need to flush its linker so that we can remove the file from disk
-	ConcertSandboxPlatformFileUtil::FlushPackageFile(InFilename);
+	// If this file maps to a package then we need to flush its linker so that we can remove the file from disk. 
+	ConcertSandboxPlatformFileUtil::FlushPackageFile(InFilename, nullptr, false);
 
 	ISourceControlProvider* SourceControlProvider = InParam.SourceControlProvider;
 	// Get the source control state of the file
