Allow IConcertClientTransactionBridge::ApplyRemoteTransaction to return optional extended information that describes what changed after applying a remote transaction

This includes a soft-reference to the objects that were modified, as well as whether they were created as part of the transaction (which is different to the event flag that said they may be created)

#rnx

diff --git a/Engine/Plugins/Developer/Concert/ConcertSync/ConcertSyncClient/Source/ConcertSyncClient/Private/ConcertClientTransactionBridge.cpp b/Engine/Plugins/Developer/Concert/ConcertSync/ConcertSyncClient/Source/ConcertSyncClient/Private/ConcertClientTransactionBridge.cpp
--- a/Engine/Plugins/Developer/Concert/ConcertSync/ConcertSyncClient/Source/ConcertSyncClient/Private/ConcertClientTransactionBridge.cpp
+++ b/Engine/Plugins/Developer/Concert/ConcertSync/ConcertSyncClient/Source/ConcertSyncClient/Private/ConcertClientTransactionBridge.cpp
@@ -332,7 +332,7 @@ struct FEditorTransactionNotification
 #endif
 
 
-void ProcessTransactionEvent(const FConcertTransactionEventBase& InEvent, const FConcertSessionVersionInfo* InVersionInfo, const TArray<FName>& InPackagesToProcess, const FConcertLocalIdentifierTable* InLocalIdentifierTablePtr, const bool bIsSnapshot, const FConcertSyncWorldRemapper& WorldRemapper, const bool bIncludeEditorOnlyProperties)
+void ProcessTransactionEvent(const FConcertTransactionEventBase& InEvent, const FConcertSessionVersionInfo* InVersionInfo, const TArray<FName>& InPackagesToProcess, const FConcertLocalIdentifierTable* InLocalIdentifierTablePtr, const bool bIsSnapshot, const FConcertSyncWorldRemapper& WorldRemapper, const bool bIncludeEditorOnlyProperties, IConcertClientTransactionBridge::FRemoteTransactionResult* ResultInfo)
 {
 	// Transactions are applied in multiple-phases...
 	//	0) Sort the objects to be processed; creation needs to happen parent->child, and update needs to happen child->parent
@@ -485,6 +485,24 @@ void ProcessTransactionEvent(const FConcertTransactionEventBase& InEvent, const
 			}
 		}
 
+		// Track the extended result info for the objects, if requested
+		if (ResultInfo)
+		{
+			for (int32 ObjectIndex = 0; ObjectIndex < SortedExportedObjects.Num(); ++ObjectIndex)
+			{
+				ConcertSyncClientUtil::FGetObjectResult& TransactionObjectRef = TransactionObjects[ObjectIndex];
+				const FConcertExportedObject& ObjectUpdate = *SortedExportedObjects[ObjectIndex];
+
+				if (TransactionObjectRef.Obj)
+				{
+					IConcertClientTransactionBridge::FRemoteTransactionObjectResult& ObjectResultInfo = ResultInfo->Objects.AddDefaulted_GetRef();
+					ObjectResultInfo.ObjectPath = TransactionObjectRef.Obj;
+					ObjectResultInfo.ObjectIndex = &ObjectUpdate - InEvent.ExportedObjects.GetData();
+					ObjectResultInfo.bNewlyCreated = NewlyCreatedObjects.Contains(TransactionObjectRef.Obj);
+				}
+			}
+		}
+
 		// Run any deferred object factory initialization
 		if (ObjectsPendingInitialization.Num() > 0)
 		{
@@ -797,14 +815,14 @@ FOnConcertConflictResolutionForPendingSend& FConcertClientTransactionBridge::OnC
 }
 
 
-void FConcertClientTransactionBridge::ApplyRemoteTransaction(const FConcertTransactionEventBase& InEvent, const FConcertSessionVersionInfo* InVersionInfo, const TArray<FName>& InPackagesToProcess, const FConcertLocalIdentifierTable* InLocalIdentifierTablePtr, const bool bIsSnapshot)
+void FConcertClientTransactionBridge::ApplyRemoteTransaction(const FConcertTransactionEventBase& InEvent, const FConcertSessionVersionInfo* InVersionInfo, const TArray<FName>& InPackagesToProcess, const FConcertLocalIdentifierTable* InLocalIdentifierTablePtr, const bool bIsSnapshot, FRemoteTransactionResult* ResultInfo)
 {
-	ConcertClientTransactionBridgeUtil::ProcessTransactionEvent(InEvent, InVersionInfo, InPackagesToProcess, InLocalIdentifierTablePtr, bIsSnapshot, FConcertSyncWorldRemapper(), bIncludeEditorOnlyProperties);
+	ConcertClientTransactionBridgeUtil::ProcessTransactionEvent(InEvent, InVersionInfo, InPackagesToProcess, InLocalIdentifierTablePtr, bIsSnapshot, FConcertSyncWorldRemapper(), bIncludeEditorOnlyProperties, ResultInfo);
 }
 
-void FConcertClientTransactionBridge::ApplyRemoteTransaction(const FConcertTransactionEventBase& InEvent, const FConcertSessionVersionInfo* InVersionInfo, const TArray<FName>& InPackagesToProcess, const FConcertLocalIdentifierTable* InLocalIdentifierTablePtr, const bool bIsSnapshot, const FConcertSyncWorldRemapper& ConcertSyncWorldRemapper)
+void FConcertClientTransactionBridge::ApplyRemoteTransaction(const FConcertTransactionEventBase& InEvent, const FConcertSessionVersionInfo* InVersionInfo, const TArray<FName>& InPackagesToProcess, const FConcertLocalIdentifierTable* InLocalIdentifierTablePtr, const bool bIsSnapshot, const FConcertSyncWorldRemapper& ConcertSyncWorldRemapper, FRemoteTransactionResult* ResultInfo)
 {
-	ConcertClientTransactionBridgeUtil::ProcessTransactionEvent(InEvent, InVersionInfo, InPackagesToProcess, InLocalIdentifierTablePtr, bIsSnapshot, ConcertSyncWorldRemapper, bIncludeEditorOnlyProperties);
+	ConcertClientTransactionBridgeUtil::ProcessTransactionEvent(InEvent, InVersionInfo, InPackagesToProcess, InLocalIdentifierTablePtr, bIsSnapshot, ConcertSyncWorldRemapper, bIncludeEditorOnlyProperties, ResultInfo);
 }
 
 bool& FConcertClientTransactionBridge::GetIgnoreLocalTransactionsRef()
diff --git a/Engine/Plugins/Developer/Concert/ConcertSync/ConcertSyncClient/Source/ConcertSyncClient/Private/ConcertClientTransactionBridge.h b/Engine/Plugins/Developer/Concert/ConcertSync/ConcertSyncClient/Source/ConcertSyncClient/Private/ConcertClientTransactionBridge.h
--- a/Engine/Plugins/Developer/Concert/ConcertSync/ConcertSyncClient/Source/ConcertSyncClient/Private/ConcertClientTransactionBridge.h
+++ b/Engine/Plugins/Developer/Concert/ConcertSync/ConcertSyncClient/Source/ConcertSyncClient/Private/ConcertClientTransactionBridge.h
@@ -27,8 +27,8 @@ public:
 	virtual FOnApplyTransaction& OnApplyTransaction() override;
 	virtual FOnConcertConflictResolutionForPendingSend& OnConflictResolutionForPendingSend() override;
 
-	virtual void ApplyRemoteTransaction(const FConcertTransactionEventBase& InEvent, const FConcertSessionVersionInfo* InVersionInfo, const TArray<FName>& InPackagesToProcess, const FConcertLocalIdentifierTable* InLocalIdentifierTablePtr, const bool bIsSnapshot) override;
-	virtual void ApplyRemoteTransaction(const FConcertTransactionEventBase& InEvent, const FConcertSessionVersionInfo* InVersionInfo, const TArray<FName>& InPackagesToProcess, const FConcertLocalIdentifierTable* InLocalIdentifierTablePtr, const bool bIsSnapshot, const FConcertSyncWorldRemapper& ConcertSyncWorldRemapper) override;
+	virtual void ApplyRemoteTransaction(const FConcertTransactionEventBase& InEvent, const FConcertSessionVersionInfo* InVersionInfo, const TArray<FName>& InPackagesToProcess, const FConcertLocalIdentifierTable* InLocalIdentifierTablePtr, const bool bIsSnapshot, FRemoteTransactionResult* ResultInfo = nullptr) override;
+	virtual void ApplyRemoteTransaction(const FConcertTransactionEventBase& InEvent, const FConcertSessionVersionInfo* InVersionInfo, const TArray<FName>& InPackagesToProcess, const FConcertLocalIdentifierTable* InLocalIdentifierTablePtr, const bool bIsSnapshot, const FConcertSyncWorldRemapper& ConcertSyncWorldRemapper, FRemoteTransactionResult* ResultInfo = nullptr) override;
 	virtual bool& GetIgnoreLocalTransactionsRef() override;
 
 	virtual void RegisterTransactionFilter(FName FilterName, FOnFilterTransactionDelegate FilterDelegate) override;
diff --git a/Engine/Plugins/Developer/Concert/ConcertSync/ConcertSyncClient/Source/ConcertSyncClient/Public/IConcertClientTransactionBridge.h b/Engine/Plugins/Developer/Concert/ConcertSync/ConcertSyncClient/Source/ConcertSyncClient/Public/IConcertClientTransactionBridge.h
--- a/Engine/Plugins/Developer/Concert/ConcertSync/ConcertSyncClient/Source/ConcertSyncClient/Public/IConcertClientTransactionBridge.h
+++ b/Engine/Plugins/Developer/Concert/ConcertSync/ConcertSyncClient/Source/ConcertSyncClient/Public/IConcertClientTransactionBridge.h
@@ -167,6 +167,25 @@ public:
 	 */
 	virtual FOnApplyTransaction& OnApplyTransaction() = 0;
 
+	/** Extended information that describes what changed for a single object within a remote transaction */
+	struct FRemoteTransactionObjectResult
+	{
+		/** Final path of the object after the given remote transaction was applied to it */
+		FSoftObjectPath ObjectPath;
+
+		/** Index into the FConcertTransactionEventBase::ExportedObjects array in the remote transaction that was applied to this object */
+		int32 ObjectIndex = INDEX_NONE;
+
+		/** True if the object was created by this remote transaction, or false if it already existed */
+		bool bNewlyCreated : 1 = false;
+	};
+
+	/** Extended information that describes what changed after applying a remote transaction */
+	struct FRemoteTransactionResult
+	{
+		TArray<FRemoteTransactionObjectResult> Objects;
+	};
+
 	/**
 	 * Apply a remote transaction event to this local instance.
 	 * @param InEvent					The event to apply.
@@ -174,8 +193,9 @@ public:
 	 * @param InPackagesToProcess		The list of packages to apply changes for, or an empty array to apply all changes.
 	 * @param InLocalIdentifierTablePtr The local identifier table for the event data (if any).
 	 * @param bIsSnapshot				True if this transaction event was a snapshot rather than a finalized transaction.
+	 * @param ResultInfo				Optional extended result information to fill.
 	 */
-	virtual void ApplyRemoteTransaction(const FConcertTransactionEventBase& InEvent, const FConcertSessionVersionInfo* InVersionInfo, const TArray<FName>& InPackagesToProcess, const FConcertLocalIdentifierTable* InLocalIdentifierTablePtr, const bool bIsSnapshot) = 0;
+	virtual void ApplyRemoteTransaction(const FConcertTransactionEventBase& InEvent, const FConcertSessionVersionInfo* InVersionInfo, const TArray<FName>& InPackagesToProcess, const FConcertLocalIdentifierTable* InLocalIdentifierTablePtr, const bool bIsSnapshot, FRemoteTransactionResult* ResultInfo = nullptr) = 0;
 
 	/**
 	 * Apply a remote transaction event to this local instance.
@@ -184,9 +204,10 @@ public:
 	 * @param InPackagesToProcess		The list of packages to apply changes for, or an empty array to apply all changes.
 	 * @param InLocalIdentifierTablePtr The local identifier table for the event data (if any).
 	 * @param bIsSnapshot				True if this transaction event was a snapshot rather than a finalized transaction.
-	 * @param ConcertSyncWorldRemapper	Remapper to use in the case the current world is different from the one sent in the transation.
+	 * @param ConcertSyncWorldRemapper	Remapper to use in the case the current world is different from the one sent in the transaction.
+	 * @param ResultInfo				Optional extended result information to fill.
 	 */
-	virtual void ApplyRemoteTransaction(const FConcertTransactionEventBase& InEvent, const FConcertSessionVersionInfo* InVersionInfo, const TArray<FName>& InPackagesToProcess, const FConcertLocalIdentifierTable* InLocalIdentifierTablePtr, const bool bIsSnapshot, const class FConcertSyncWorldRemapper& ConcertSyncWorldRemapper) = 0;
+	virtual void ApplyRemoteTransaction(const FConcertTransactionEventBase& InEvent, const FConcertSessionVersionInfo* InVersionInfo, const TArray<FName>& InPackagesToProcess, const FConcertLocalIdentifierTable* InLocalIdentifierTablePtr, const bool bIsSnapshot, const class FConcertSyncWorldRemapper& ConcertSyncWorldRemapper, FRemoteTransactionResult* ResultInfo = nullptr) = 0;
 
 	/** Callback to register delegate for handling transaction events */
 	UE_DEPRECATED(5.5, "Use the version of RegisterTransactionFilter accepting FOnFilterTransactionDelegate instead.")
