Debug commands for inspecting the transaction state of the editor.

#rnx

diff --git a/Engine/Plugins/Developer/Concert/ConcertSync/ConcertSyncClient/Source/ConcertSyncClient/Private/ConcertClientTransactionBridge.cpp b/Engine/Plugins/Developer/Concert/ConcertSync/ConcertSyncClient/Source/ConcertSyncClient/Private/ConcertClientTransactionBridge.cpp
--- a/Engine/Plugins/Developer/Concert/ConcertSync/ConcertSyncClient/Source/ConcertSyncClient/Private/ConcertClientTransactionBridge.cpp
+++ b/Engine/Plugins/Developer/Concert/ConcertSync/ConcertSyncClient/Source/ConcertSyncClient/Private/ConcertClientTransactionBridge.cpp
@@ -22,6 +22,7 @@
 #include "UObject/SoftObjectPath.h"
 #include "UObject/WeakObjectPtr.h"
 
+
 #if WITH_EDITOR
 	#include "Editor.h"
 	#include "UnrealEdGlobals.h"
@@ -29,6 +30,8 @@
 	#include "Editor/TransBuffer.h"
 	#include "Engine/Selection.h"
 
+	#include "Framework/Commands/UIAction.h"
+
 	#include "Elements/Framework/EngineElementsLibrary.h"
 	#include "Elements/Framework/TypedElementSelectionSet.h"
 #endif
@@ -45,6 +48,33 @@ static TAutoConsoleVariable<int32> CVarIgnoreTransactionIncludeFilter(TEXT("Conc
 static TAutoConsoleVariable<int32> CVarAlwaysSuspendBroadcastUndoRedo(TEXT("Concert.AlwaysSuspendBroadcastUndoRedo"), 0,
 																	  TEXT("Always suspend the undo/redo broadcast message. Can help with transactions that steal user focus."));
 
+#if WITH_EDITOR
+
+void ReportCurrentTransactionState()
+{
+	UE_LOG(LogConcert, Display, TEXT("Undo/Redo GIsTransacting is %s."), GIsTransacting ? TEXT("true") : TEXT("false"));
+	UE_LOG(LogConcert, Display, TEXT("GEditor->CanTransact() is %s."), GEditor && GEditor->CanTransact() ? TEXT("true") : TEXT("false"));
+	UE_LOG(LogConcert, Display, TEXT("GEditor->IsTransactionActive() is %s."), GEditor && GEditor->IsTransactionActive() ? TEXT("true") : TEXT("false"));
+}
+
+void ResetTransaction()
+{
+	if (GEditor)
+	{
+		GEditor->ResetTransaction(LOCTEXT("ConcertDebugCvarReset", "Reset by console variable command."));
+	}
+}
+
+FAutoConsoleCommand CVarReportCurrentTransactionState(TEXT("Concert.Debug.ReportTransactionState"),
+	TEXT("Report to the log file the current transaction state."),
+	FExecuteAction::CreateStatic(&ReportCurrentTransactionState));
+
+FAutoConsoleCommand CVarResetTransaction(TEXT("Concert.Debug.ResetTransaction"),
+	TEXT("Reset the current open transaction."),
+	FExecuteAction::CreateStatic(&ResetTransaction));
+
+#endif // WITH_EDITOR
+	   
 UClass* LoadClassHelper(const FSoftClassPath& ClassPath)
 {
 	static TMap<FSoftClassPath, TWeakObjectPtr<UClass>> ClassCache;
