Multiuser: Fix possible crash when stopping a take record or reverting take changes.

diff --git a/Engine/Plugins/Developer/Concert/ConcertSync/ConcertSyncClient/Source/ConcertSyncClient/Private/ConcertClientSequencerManager.cpp b/Engine/Plugins/Developer/Concert/ConcertSync/ConcertSyncClient/Source/ConcertSyncClient/Private/ConcertClientSequencerManager.cpp
--- a/Engine/Plugins/Developer/Concert/ConcertSync/ConcertSyncClient/Source/ConcertSyncClient/Private/ConcertClientSequencerManager.cpp
+++ b/Engine/Plugins/Developer/Concert/ConcertSync/ConcertSyncClient/Source/ConcertSyncClient/Private/ConcertClientSequencerManager.cpp
@@ -951,39 +951,40 @@ void FConcertClientSequencerManager::ApplyOpenEvent(const FConcertSequencerOpenE
 
 	if (GIsEditor && bDidOpen && UE::Private::ConcertClientSequencerManager::IsPendingTakePath(SequenceObjectPath))
 	{
-		ULevelSequence* PendingLevelSequence = FindObject<ULevelSequence>(nullptr, *UE::Private::ConcertClientSequencerManager::PendingTakePath);
+		if (ULevelSequence* PendingLevelSequence = FindObject<ULevelSequence>(nullptr, *UE::Private::ConcertClientSequencerManager::PendingTakePath))
+		{
+			// Apply pending take data to our version of the pending take.
+			//
+			FConcertSyncEncounteredMissingObject MissingObjectDelegate;
+			MissingObjectDelegate.BindLambda([](const FStringView MissingObject)
+				{
+					UE_LOG(LogConcertSequencerSync, Display, TEXT("Missing Object %s when loading PendingTake"), *FString(MissingObject));
+				});
 
-		// Apply pending take data to our version of the pending take.
-		//
-		FConcertSyncEncounteredMissingObject MissingObjectDelegate;
-		MissingObjectDelegate.BindLambda([](const FStringView MissingObject)
+			if (InOpenEvent.TakeData.Bytes.Num() > 0)
 			{
-				UE_LOG(LogConcertSequencerSync, Display, TEXT("Missing Object %s when loading PendingTake"), *FString(MissingObject));
-			});
-
-		if (InOpenEvent.TakeData.Bytes.Num() > 0)
-		{
-			FConcertSyncWorldRemapper Remapper(
-				TEXT("/Engine/Transient.__PendingLevelSequence__"), PendingLevelSequence->GetPathName());
-			FConcertSyncObjectReader Reader(nullptr, MoveTemp(Remapper), nullptr, PendingLevelSequence, InOpenEvent.TakeData.Bytes,
-											MissingObjectDelegate);
-			Reader.SetSerializeNestedObjects(true);
-			Reader.SerializeObject(PendingLevelSequence);
-		}
-		else
-		{
-			UE_LOG(LogConcertSequencerSync, Display, TEXT("Missing take data on pending take."));
-		}
+				FConcertSyncWorldRemapper Remapper(
+					TEXT("/Engine/Transient.__PendingLevelSequence__"), PendingLevelSequence->GetPathName());
+				FConcertSyncObjectReader Reader(nullptr, MoveTemp(Remapper), nullptr, PendingLevelSequence, InOpenEvent.TakeData.Bytes,
+												MissingObjectDelegate);
+				Reader.SetSerializeNestedObjects(true);
+				Reader.SerializeObject(PendingLevelSequence);
+			}
+			else
+			{
+				UE_LOG(LogConcertSequencerSync, Display, TEXT("Missing take data on pending take."));
+			}
 
-		#if WITH_EDITOR
-		{
-			// We must specify the pending level sequence and not the pending take.  If we specify the pending take after we serialized into the existing
-			// asset the open asset will cause take recorder tab to assume that we want to edit the pending take where we really just want to view
-			// the level sequence.
-			GEditor->GetEditorSubsystem<UAssetEditorSubsystem>()->OpenEditorForAsset(PendingLevelSequence);
+#if WITH_EDITOR
+			{
+				// We must specify the pending level sequence and not the pending take.  If we specify the pending take after we serialized into the existing
+				// asset the open asset will cause take recorder tab to assume that we want to edit the pending take where we really just want to view
+				// the level sequence.
+				GEditor->GetEditorSubsystem<UAssetEditorSubsystem>()->OpenEditorForAsset(PendingLevelSequence);
+			}
+#endif
+			PendingLevelSequence->GetPackage()->SetDirtyFlag(false);
 		}
-		#endif
-		PendingLevelSequence->GetPackage()->SetDirtyFlag(false);
 	}
 }
 
diff --git a/Engine/Plugins/Developer/Concert/ConcertSync/ConcertSyncCore/Source/ConcertSyncCore/Private/ConcertSyncArchives.cpp b/Engine/Plugins/Developer/Concert/ConcertSync/ConcertSyncCore/Source/ConcertSyncCore/Private/ConcertSyncArchives.cpp
--- a/Engine/Plugins/Developer/Concert/ConcertSync/ConcertSyncCore/Source/ConcertSyncCore/Private/ConcertSyncArchives.cpp
+++ b/Engine/Plugins/Developer/Concert/ConcertSync/ConcertSyncCore/Source/ConcertSyncCore/Private/ConcertSyncArchives.cpp
@@ -547,7 +547,10 @@ FConcertSyncObjectReader::FConcertSyncObjectReader(const FConcertLocalIdentifier
 void FConcertSyncObjectReader::SerializeObject(UObject* InObject)
 {
 	CurrentOuter = InObject;
-	InObject->Serialize(*this);
+	if (InObject)
+	{
+		InObject->Serialize(*this);
+	}
 }
 
 void FConcertSyncObjectReader::SerializeProperty(const FProperty* InProp, UObject* InObject)
