Concert sequencer manager fix for garbage collection on exit.

#rnx

diff --git a/Engine/Plugins/Developer/Concert/ConcertSync/ConcertSyncClient/Source/ConcertSyncClient/Private/ConcertClientSequencerManager.cpp b/Engine/Plugins/Developer/Concert/ConcertSync/ConcertSyncClient/Source/ConcertSyncClient/Private/ConcertClientSequencerManager.cpp
--- a/Engine/Plugins/Developer/Concert/ConcertSync/ConcertSyncClient/Source/ConcertSyncClient/Private/ConcertClientSequencerManager.cpp
+++ b/Engine/Plugins/Developer/Concert/ConcertSync/ConcertSyncClient/Source/ConcertSyncClient/Private/ConcertClientSequencerManager.cpp
@@ -393,7 +393,7 @@ ETransactionFilterResult FConcertClientSequencerManager::ShouldObjectBeTransacte
 	if (bShouldCheckOutermost)
 	{
 		UObject* Outermost = IsValid(FilterArgs.ObjectToFilter) ? FilterArgs.ObjectToFilter->GetOutermostObject() : nullptr;
-		if (WeakSession.IsValid() && Outermost && Outermost->IsA<ULevelSequence>())
+		if (WeakSession.IsValid() && Outermost && Outermost->IsA<ULevelSequence>() && !Outermost->HasAnyFlags(RF_Transient))
 		{
 			return ETransactionFilterResult::IncludeObject;
 		}
@@ -463,6 +463,7 @@ void FConcertClientSequencerManager::Unregister(TSharedRef<IConcertClientSession
 		}
 	}
 
+	SequencePlayers.Reset();
 	WeakSession.Reset();
 
 	Preloader->OnUnregister(InSession);
