Address issue found on unregister and explicitly destroy any open players in -game instances.

diff --git a/Engine/Plugins/Developer/Concert/ConcertSync/ConcertSyncClient/Source/ConcertSyncClient/Private/ConcertClientSequencerManager.cpp b/Engine/Plugins/Developer/Concert/ConcertSync/ConcertSyncClient/Source/ConcertSyncClient/Private/ConcertClientSequencerManager.cpp
--- a/Engine/Plugins/Developer/Concert/ConcertSync/ConcertSyncClient/Source/ConcertSyncClient/Private/ConcertClientSequencerManager.cpp
+++ b/Engine/Plugins/Developer/Concert/ConcertSync/ConcertSyncClient/Source/ConcertSyncClient/Private/ConcertClientSequencerManager.cpp
@@ -463,6 +463,15 @@ void FConcertClientSequencerManager::Unregister(TSharedRef<IConcertClientSession
 		}
 	}
 
+	for (TTuple<FName, TObjectPtr<ALevelSequenceActor>>& Item : SequencePlayers)
+	{
+		ALevelSequenceActor* LevelSequenceActor = Item.Value;
+		if (LevelSequenceActor)
+		{
+			DestroyPlayer(LevelSequenceActor);
+		}				
+	}
+
 	SequencePlayers.Reset();
 	WeakSession.Reset();
 
