Multi User: Fix crash when deleting Level Sequence in a MU session.

The issue was that NewObject was being called while iterating through the global UObject hash table (ForEachObjectWithPackage).
As ForEachObjectWithPackage was deleting the assets. this caused the asset editor tools to close the underlying asset editor which in turn calls NewObject for UToolMenus.

diff --git a/Engine/Plugins/Developer/Concert/ConcertSync/ConcertSyncClient/Source/ConcertSyncClient/Private/ConcertSyncClientUtil.cpp b/Engine/Plugins/Developer/Concert/ConcertSync/ConcertSyncClient/Source/ConcertSyncClient/Private/ConcertSyncClientUtil.cpp
--- a/Engine/Plugins/Developer/Concert/ConcertSync/ConcertSyncClient/Source/ConcertSyncClient/Private/ConcertSyncClientUtil.cpp
+++ b/Engine/Plugins/Developer/Concert/ConcertSync/ConcertSyncClient/Source/ConcertSyncClient/Private/ConcertSyncClientUtil.cpp
@@ -750,11 +750,20 @@ void PurgePackages(TArrayView<const FName> InPackageNames)
 			ExistingPackage->SetDirtyFlag(false);
 
 			CollectObjectToPurge(ExistingPackage);
-			ForEachObjectWithPackage(ExistingPackage, [&CollectObjectToPurge](UObject* InObject)
+
+			// It's unsafe to call CollectObjectToPurge in ForEachObjectWithPackage lambda.
+			// CollectObjectToPurge invokes global callbacks (e.g. closing asset editor), which may cause a NewObject call.
+			// NewObject calls are illegal while iterating through the global UObject hash map.
+			TArray<UObject*> ObjectsInPackage;
+			ForEachObjectWithPackage(ExistingPackage, [&ObjectsInPackage](UObject* InObject)
 			{
-				CollectObjectToPurge(InObject);
+				ObjectsInPackage.Add(InObject);
 				return true;
 			});
+			for (UObject* Subobject : ObjectsInPackage)
+			{
+				CollectObjectToPurge(Subobject);
+			}
 
 			bEditedMapPurged |= EditedMapPackage == ExistingPackage;
 		}
