Fixing issue where folders created in a Concert Sandbox could not be browsed to with "Show In Explorer" until they had been persisted.

diff --git a/Engine/Plugins/Developer/Concert/ConcertSync/ConcertSyncClient/Source/ConcertSyncClient/Private/ConcertSandboxPlatformFile.cpp b/Engine/Plugins/Developer/Concert/ConcertSync/ConcertSyncClient/Source/ConcertSyncClient/Private/ConcertSandboxPlatformFile.cpp
--- a/Engine/Plugins/Developer/Concert/ConcertSync/ConcertSyncClient/Source/ConcertSyncClient/Private/ConcertSandboxPlatformFile.cpp
+++ b/Engine/Plugins/Developer/Concert/ConcertSync/ConcertSyncClient/Source/ConcertSyncClient/Private/ConcertSandboxPlatformFile.cpp
@@ -648,7 +648,7 @@ FString FConcertSandboxPlatformFile::ConvertToAbsolutePathForExternalAppForRead(
 
 	if (ResolvedPath.HasSandboxPath())
 	{
-		if (IsPathDeleted(ResolvedPath) || LowerLevel->FileExists(*ResolvedPath.GetSandboxPath()))
+		if (IsPathDeleted(ResolvedPath) || LowerLevel->FileExists(*ResolvedPath.GetSandboxPath()) || LowerLevel->DirectoryExists(*ResolvedPath.GetSandboxPath()))
 		{
 			return LowerLevel->ConvertToAbsolutePathForExternalAppForRead(*ResolvedPath.GetSandboxPath());
 		}
