Allowing concert and similar systems to attempt a delete+copy if a copy-over fails when flushing a sandbox. This is useful to deal with file sharing issues if the file is opened SHARE_DELETE but not SHARE_WRITE

diff --git a/Engine/Plugins/Developer/Concert/ConcertSync/ConcertSyncClient/Source/ConcertSyncClient/Private/ConcertSandboxPlatformFile.cpp b/Engine/Plugins/Developer/Concert/ConcertSync/ConcertSyncClient/Source/ConcertSyncClient/Private/ConcertSandboxPlatformFile.cpp
--- a/Engine/Plugins/Developer/Concert/ConcertSync/ConcertSyncClient/Source/ConcertSyncClient/Private/ConcertSandboxPlatformFile.cpp
+++ b/Engine/Plugins/Developer/Concert/ConcertSync/ConcertSyncClient/Source/ConcertSyncClient/Private/ConcertSandboxPlatformFile.cpp
@@ -735,7 +735,7 @@ bool FConcertSandboxPlatformFile::CopyFileWithSCC(
 
 	// Copy the on-disk sandbox file
 	FString ToFileDir = FPaths::GetPath(InToFilename);
-	if (!CreateDirectoryTree(*ToFileDir) || !LowerLevel->CopyFile(InToFilename, InFromFilename))
+	if (!CreateDirectoryTree(*ToFileDir) || !LowerLevel->CopyFile(InToFilename, InFromFilename, EPlatformFileRead::None, EPlatformFileWrite::AttemptDeleteAndCreate))
 	{
 		FText Failure = FText::Format(
 			NSLOCTEXT("Concert", "CopyFileFailure", "Failed to copy file '{0}' (from '{1}') when persiting sandbox state!"),
diff --git a/Engine/Source/Runtime/Core/Private/GenericPlatform/GenericPlatformFile.cpp b/Engine/Source/Runtime/Core/Private/GenericPlatform/GenericPlatformFile.cpp
--- a/Engine/Source/Runtime/Core/Private/GenericPlatform/GenericPlatformFile.cpp
+++ b/Engine/Source/Runtime/Core/Private/GenericPlatform/GenericPlatformFile.cpp
@@ -879,7 +879,19 @@ bool IPlatformFile::CopyFile(const TCHAR* To, const TCHAR* From, EPlatformFileRe
 	TUniquePtr<IFileHandle> ToFile(OpenWrite(To, false, (WriteFlags & EPlatformFileWrite::AllowRead) != EPlatformFileWrite::None));
 	if (!ToFile)
 	{
-		return false;
+		bool bFailed = true;
+		if ((WriteFlags & EPlatformFileWrite::AttemptDeleteAndCreate) != EPlatformFileWrite::None)
+		{
+			if (DeleteFile(To))
+			{
+				ToFile = TUniquePtr<IFileHandle>(OpenWrite(To, false, (WriteFlags & EPlatformFileWrite::AllowRead) != EPlatformFileWrite::None));
+				bFailed = (!ToFile);
+			}
+		}
+		if (bFailed)
+		{
+			return false;
+		}
 	}
 	int64 Size = FromFile->Size();
 	if (Size < 1)
diff --git a/Engine/Source/Runtime/Core/Public/GenericPlatform/GenericPlatformFile.h b/Engine/Source/Runtime/Core/Public/GenericPlatform/GenericPlatformFile.h
--- a/Engine/Source/Runtime/Core/Public/GenericPlatform/GenericPlatformFile.h
+++ b/Engine/Source/Runtime/Core/Public/GenericPlatform/GenericPlatformFile.h
@@ -69,7 +69,8 @@ ENUM_CLASS_FLAGS(EPlatformFileRead);
 enum class EPlatformFileWrite : uint8
 {
 	None = 0x0,
-	AllowRead = 0x01	// attempts to open for write while allowing others to read
+	AllowRead = 0x01,	// attempts to open for write while allowing others to read
+	AttemptDeleteAndCreate = 0x02 // if the target file cannot be open for write, attempt to delete it and create a new one in it splace
 };
 
 ENUM_CLASS_FLAGS(EPlatformFileWrite);
