Concert: Removed PersistedFiles TMap from FConcertClientWorkspaceData and FConcertSandboxPlatformFile, and the related dead code.

PersistedFiles no longer affected any decision making. It was only being loaded from disk and then saved into FConcertSandboxPlatformFile. The only place that really used it was FConcertSandboxPlatformFile::GatherChangedFileNames, but that function is unreferenced by any code.

Another reason we can be sure this is safe is because FConcertClientWorkspaceData was being saved into a directory containg FApp::GetInstanceId as path. That means after you restart the editor and rejoin the session, the data would no longer be loaded either.

#Synced-CL 47637010

diff --git a/Engine/Plugins/Developer/Concert/ConcertSync/ConcertSyncClient/Source/ConcertSyncClient/Private/ConcertClientPackageManager.cpp b/Engine/Plugins/Developer/Concert/ConcertSync/ConcertSyncClient/Source/ConcertSyncClient/Private/ConcertClientPackageManager.cpp
--- a/Engine/Plugins/Developer/Concert/ConcertSync/ConcertSyncClient/Source/ConcertSyncClient/Private/ConcertClientPackageManager.cpp
+++ b/Engine/Plugins/Developer/Concert/ConcertSync/ConcertSyncClient/Source/ConcertSyncClient/Private/ConcertClientPackageManager.cpp
@@ -174,67 +174,6 @@ bool FConcertClientPackageManager::ShouldIgnorePackageDirtyEvent(class UPackage*
 		|| (!PassesPackageFilters(InPackage));
 }
 
-TMap<FString, int64> FConcertClientPackageManager::GetPersistedFiles() const
-{
-	TMap<FString, int64> PersistedFiles;
-#if WITH_EDITOR
-	if (SandboxPlatformFile)
-	{
-		FString PackagePath;
-		for (const auto& PersistedFilePair : SandboxPlatformFile->GetPersistedFiles())
-		{
-			if (FPackageName::TryConvertFilenameToLongPackageName(PersistedFilePair.Key, PackagePath))
-			{
-				int64 PackageRevision = 0;
-				if (LiveSession->GetSessionDatabase().GetPackageHeadRevision(*PackagePath, PackageRevision))
-				{
-					PersistedFiles.Add(PackagePath, PackageRevision);
-				}
-			}
-		}
-	}
-#endif
-	return PersistedFiles;
-}
-
-void FConcertClientPackageManager::SynchronizePersistedFiles(const TMap<FString, int64>& PersistedFiles)
-{
-#if WITH_EDITOR
-	if (SandboxPlatformFile)
-	{
-		auto GetPackageFilenameForRevision = [this](const FString& PackageName, const int64 PackageRevision, FString& OutPackageFilename) -> bool
-		{
-			FConcertPackageInfo PackageInfo;
-			if (LiveSession->GetSessionDatabase().GetPackageInfoForRevision(*PackageName, PackageInfo, &PackageRevision))
-			{
-				if (FPackageName::TryConvertLongPackageNameToFilename(PackageName, OutPackageFilename, PackageInfo.PackageFileExtension))
-				{
-					OutPackageFilename = FPaths::ConvertRelativePathToFull(MoveTemp(OutPackageFilename));
-					return true;
-				}
-			}
-			return false;
-		};
-
-		FString PackageFilename;
-		TArray<FString> PersistedFilePaths;
-		for (const auto& PersistedFilePair : PersistedFiles)
-		{
-			int64 PackageRevision = 0;
-			if (LiveSession->GetSessionDatabase().GetPackageHeadRevision(*PersistedFilePair.Key, PackageRevision))
-			{
-				// If the current package ledger head revision match the persisted file revision, add the file as persisted
-				if (PackageRevision == PersistedFilePair.Value && GetPackageFilenameForRevision(PersistedFilePair.Key, PackageRevision, PackageFilename))
-				{
-					PersistedFilePaths.Add(PackageFilename);
-				}
-			}
-		}
-		SandboxPlatformFile->AddFilesAsPersisted(PersistedFilePaths);
-	}
-#endif
-}
-
 void FConcertClientPackageManager::QueueDirtyPackagesForReload()
 {
 	TArray<UPackage*> DirtyPkgs;
diff --git a/Engine/Plugins/Developer/Concert/ConcertSync/ConcertSyncClient/Source/ConcertSyncClient/Private/ConcertClientPackageManager.h b/Engine/Plugins/Developer/Concert/ConcertSync/ConcertSyncClient/Source/ConcertSyncClient/Private/ConcertClientPackageManager.h
--- a/Engine/Plugins/Developer/Concert/ConcertSync/ConcertSyncClient/Source/ConcertSyncClient/Private/ConcertClientPackageManager.h
+++ b/Engine/Plugins/Developer/Concert/ConcertSync/ConcertSyncClient/Source/ConcertSyncClient/Private/ConcertClientPackageManager.h
@@ -42,18 +42,7 @@ public:
 	 * @return true if dirty even should be ignored for InPackage
 	 */
 	bool ShouldIgnorePackageDirtyEvent(class UPackage* InPackage) const;
-
-	/** 
-	 * @return the map of persisted files to their current package ledger version.
-	 */
-	TMap<FString, int64> GetPersistedFiles() const;
-
-	/**
-	 * Synchronize files that should be considered as already persisted from session.
-	 * @param PersistedFiles Map of persisted files to their package version to mark as persisted if their ledger version match.
-	 */
-	void SynchronizePersistedFiles(const TMap<FString, int64>& PersistedFiles);
-
+	
 	/**
 	 * Discard dirty packages changes by queuing them for hot reload
 	 */
diff --git a/Engine/Plugins/Developer/Concert/ConcertSync/ConcertSyncClient/Source/ConcertSyncClient/Private/ConcertClientWorkspace.cpp b/Engine/Plugins/Developer/Concert/ConcertSync/ConcertSyncClient/Source/ConcertSyncClient/Private/ConcertClientWorkspace.cpp
--- a/Engine/Plugins/Developer/Concert/ConcertSync/ConcertSyncClient/Source/ConcertSyncClient/Private/ConcertClientWorkspace.cpp
+++ b/Engine/Plugins/Developer/Concert/ConcertSync/ConcertSyncClient/Source/ConcertSyncClient/Private/ConcertClientWorkspace.cpp
@@ -640,17 +640,12 @@ void FConcertClientWorkspace::LoadSessionData()
 	if (SessionData.SessionIdentifier != LiveSession->GetSession().GetSessionServerEndpointId())
 	{
 		SessionData.SessionIdentifier.Invalidate();
-		SessionData.PersistedFiles.Empty();
 	}
 }
 
 void FConcertClientWorkspace::SaveSessionData()
 {
 	SessionData.SessionIdentifier = LiveSession->GetSession().GetSessionServerEndpointId();
-	if (PackageManager)
-	{
-		SessionData.PersistedFiles = PackageManager->GetPersistedFiles();
-	}
 	
 	FString ClientWorkspaceDataPath = LiveSession->GetSession().GetSessionWorkingDirectory() / TEXT("WorkspaceData.json");
 	if (TUniquePtr<FArchive> FileWriter = TUniquePtr<FArchive>(IFileManager::Get().CreateFileWriter(*ClientWorkspaceDataPath)))
@@ -888,7 +883,6 @@ void FConcertClientWorkspace::OnEndFrame()
 		}
 		if (PackageManager)
 		{
-			PackageManager->SynchronizePersistedFiles(SessionData.PersistedFiles);
 			PackageManager->QueueDirtyPackagesForReload();
 			PackageManager->ApplyAllHeadPackageData();
 			PackageManager->SynchronizeInMemoryPackages();
diff --git a/Engine/Plugins/Developer/Concert/ConcertSync/ConcertSyncClient/Source/ConcertSyncClient/Private/ConcertClientWorkspaceData.h b/Engine/Plugins/Developer/Concert/ConcertSync/ConcertSyncClient/Source/ConcertSyncClient/Private/ConcertClientWorkspaceData.h
--- a/Engine/Plugins/Developer/Concert/ConcertSync/ConcertSyncClient/Source/ConcertSyncClient/Private/ConcertClientWorkspaceData.h
+++ b/Engine/Plugins/Developer/Concert/ConcertSync/ConcertSyncClient/Source/ConcertSyncClient/Private/ConcertClientWorkspaceData.h
@@ -17,8 +17,4 @@ struct FConcertClientWorkspaceData
 	/** The session identifier that this client data is associated with. */
 	UPROPERTY()
 	FGuid SessionIdentifier;
-
-	/** Map of packages persisted locally from the session associated with their package revision.*/
-	UPROPERTY()
-	TMap<FString, int64> PersistedFiles;
 };
diff --git a/Engine/Plugins/Developer/Concert/ConcertSync/ConcertSyncClient/Source/ConcertSyncClient/Private/ConcertSandboxPlatformFile.cpp b/Engine/Plugins/Developer/Concert/ConcertSync/ConcertSyncClient/Source/ConcertSyncClient/Private/ConcertSandboxPlatformFile.cpp
--- a/Engine/Plugins/Developer/Concert/ConcertSync/ConcertSyncClient/Source/ConcertSyncClient/Private/ConcertSandboxPlatformFile.cpp
+++ b/Engine/Plugins/Developer/Concert/ConcertSync/ConcertSyncClient/Source/ConcertSyncClient/Private/ConcertSandboxPlatformFile.cpp
@@ -958,75 +958,6 @@ void FConcertSandboxPlatformFile::DiscardSandbox(TArray<FName>& OutPackagesPendi
 #endif
 }
 
-void FConcertSandboxPlatformFile::AddFilesAsPersisted(TArrayView<const FString> InFiles)
-{
-	for (const FString& File : InFiles)
-	{
-		PersistedFiles.Add(File, GetTimeStamp(*File));
-	}
-}
-
-void FConcertSandboxPlatformFile::RemoveFilesAsPersisted(TArrayView<const FString> InFiles)
-{
-	for (const FString& File : InFiles)
-	{
-		PersistedFiles.Remove(File);
-	}
-}
-
-const TMap<FString, FDateTime>& FConcertSandboxPlatformFile::GetPersistedFiles() const
-{
-	return PersistedFiles;
-}
-
-TArray<FString> FConcertSandboxPlatformFile::GatherSandboxChangedFilenames() const
-{
-	TArray<FString> ChangedFiles;
-	
-	// Gather deleted paths
-	{
-		FScopeLock DeletedSandboxPathsLock(&DeletedSandboxPathsCS);
-		ChangedFiles.Reserve(DeletedSandboxPaths.Num());
-		Algo::TransformIf(DeletedSandboxPaths, ChangedFiles,
-			[this](const FConcertSandboxPlatformFilePath& DeletedPath)
-			{
-				return LowerLevel->FileExists(*DeletedPath.GetNonSandboxPath());
-			},
-			[this](const FConcertSandboxPlatformFilePath& DeletedPath)
-			{
-				return DeletedPath.GetNonSandboxPath();
-			});
-	}
-
-	// Gather mounts files
-	{
-		FScopeLock SandboxMountPointsLock(&SandboxMountPointsCS);
-		for (const FSandboxMountPoint& MountPoint : SandboxMountPoints)
-		{
-			LowerLevel->IterateDirectoryRecursively(*MountPoint.Path.GetSandboxPath(), [this, &MountPoint, &ChangedFiles](const TCHAR* InFilenameOrDirectory, const bool InIsDirectory) -> bool
-			{
-				if (!InIsDirectory)
-				{
-					ChangedFiles.Add(FConcertSandboxPlatformFilePath::CreateNonSandboxPath(FPaths::ConvertRelativePathToFull(InFilenameOrDirectory), MountPoint.Path).GetNonSandboxPath());
-				}
-				return true; // continue iteration
-			});
-		}
-	}
-
-	// Remove files marked as persisted already
-	ChangedFiles.RemoveAllSwap([this](const FString& InFile)
-	{
-		if (const FDateTime* ModifiedTime = PersistedFiles.Find(InFile))
-		{
-			return const_cast<FConcertSandboxPlatformFile*>(this)->GetTimeStamp(*InFile) <= *ModifiedTime;
-		}
-		return false;
-	});
-
-	return ChangedFiles;
-}
-
 void FConcertSandboxPlatformFile::OnContentPathMounted(const FString& InAssetPath, const FString& InFilesystemPath)
 {
 	RegisterContentMountPath(InAssetPath, InFilesystemPath);
diff --git a/Engine/Plugins/Developer/Concert/ConcertSync/ConcertSyncClient/Source/ConcertSyncClient/Private/ConcertSandboxPlatformFile.h b/Engine/Plugins/Developer/Concert/ConcertSync/ConcertSyncClient/Source/ConcertSyncClient/Private/ConcertSandboxPlatformFile.h
--- a/Engine/Plugins/Developer/Concert/ConcertSync/ConcertSyncClient/Source/ConcertSyncClient/Private/ConcertSandboxPlatformFile.h
+++ b/Engine/Plugins/Developer/Concert/ConcertSync/ConcertSyncClient/Source/ConcertSyncClient/Private/ConcertSandboxPlatformFile.h
@@ -170,36 +170,6 @@ public:
 	 */
 	void DiscardSandbox(TArray<FName>& OutPackagesPendingHotReload, TArray<FName>& OutPackagesPendingPurge);
 
-	/**
-	 *	Add a list of files as persisted from the sandbox.
-	 *	Persisted files will will ignored when returning changed files 
-	 *	from the sandbox if the modified date match.
-	 *
-	 *	@param InFiles		List of files to add to the persisted list.
-	 */
-	void AddFilesAsPersisted(TArrayView<const FString> InFiles);
-
-	/**
-	 *	Remove a list of files as persisted from the sandbox.
-	 *	@see AddFilesAsPersisted, PersistSandbox
-	 *
-	 *	@param InFiles		Array of files to remove
-	 */
-	void RemoveFilesAsPersisted(TArrayView<const FString> InFiles);
-
-	/**
-	 *	Return the list of files currently in the persisted list.
-	 *	@see AddFilesAsPersisted, PersistSandbox
-	 */
-	const TMap<FString, FDateTime>& GetPersistedFiles() const;
-
-	/** 
-	 *	Gather all files changes that are currently in the sandbox
-	 *	while ignoring files marked as persisted.
-	 *	@return Array of filepaths in not particular order.
-	 */
-	TArray<FString> GatherSandboxChangedFilenames() const;
-
 	/**
 	 * Returns true if the given package file exists on non sandbox path.
 	 */
