Concert: Convert ConcertClient to use new FileSandbox API for package sandbox.

#Synced-CL 47637010

diff --git a/Engine/Plugins/Developer/Concert/ConcertSync/ConcertSyncClient/ConcertSyncClient.uplugin b/Engine/Plugins/Developer/Concert/ConcertSync/ConcertSyncClient/ConcertSyncClient.uplugin
--- a/Engine/Plugins/Developer/Concert/ConcertSync/ConcertSyncClient/ConcertSyncClient.uplugin
+++ b/Engine/Plugins/Developer/Concert/ConcertSync/ConcertSyncClient/ConcertSyncClient.uplugin
@@ -31,6 +31,10 @@
 		{
 			"Name": "ConcertMain",
 			"Enabled": true
+		},
+		{
+			"Name": "FileSandbox",
+			"Enabled": true
 		}
 	]
 }
diff --git a/Engine/Plugins/Developer/Concert/ConcertSync/ConcertSyncClient/Source/ConcertSyncClient/ConcertSyncClient.Build.cs b/Engine/Plugins/Developer/Concert/ConcertSync/ConcertSyncClient/Source/ConcertSyncClient/ConcertSyncClient.Build.cs
--- a/Engine/Plugins/Developer/Concert/ConcertSync/ConcertSyncClient/Source/ConcertSyncClient/ConcertSyncClient.Build.cs
+++ b/Engine/Plugins/Developer/Concert/ConcertSync/ConcertSyncClient/Source/ConcertSyncClient/ConcertSyncClient.Build.cs
@@ -23,6 +23,7 @@ namespace UnrealBuildTool.Rules
 				{
 					"AssetRegistry",
 					"ConcertTransport",
+					"FileSandboxCore",
 					"HeadMountedDisplay",
 					"InputCore",
 					"JsonUtilities",
diff --git a/Engine/Plugins/Developer/Concert/ConcertSync/ConcertSyncClient/Source/ConcertSyncClient/Private/ConcertClientPackageManager.cpp b/Engine/Plugins/Developer/Concert/ConcertSync/ConcertSyncClient/Source/ConcertSyncClient/Private/ConcertClientPackageManager.cpp
--- a/Engine/Plugins/Developer/Concert/ConcertSync/ConcertSyncClient/Source/ConcertSyncClient/Private/ConcertClientPackageManager.cpp
+++ b/Engine/Plugins/Developer/Concert/ConcertSync/ConcertSyncClient/Source/ConcertSyncClient/Private/ConcertClientPackageManager.cpp
@@ -27,6 +27,8 @@
 #include "HAL/PlatformFileManager.h"
 
 #include "ISourceControlModule.h"
+#include "Sandbox/ConcertPackageSandbox.h"
+#include "Sandbox/IConcertSandbox.h"
 #include "UObject/UObjectHash.h"
 
 #if WITH_EDITOR
@@ -70,6 +72,8 @@ FConcertClientPackageManager::FConcertClientPackageManager(TSharedRef<FConcertSy
 	, bIgnorePackageDirtyEvent(false)
 	, FileSharingService(MoveTemp(InFileSharingService))
 {
+	using namespace UE::ConcertClient;
+	
 	check(LiveSession->IsValidSession());
 	check(EnumHasAnyFlags(LiveSession->GetSessionFlags(), EConcertSyncSessionFlags::EnablePackages));
 	check(PackageBridge);
@@ -78,8 +82,7 @@ FConcertClientPackageManager::FConcertClientPackageManager(TSharedRef<FConcertSy
 	if (EnumHasAnyFlags(LiveSession->GetSessionFlags(), EConcertSyncSessionFlags::ShouldUsePackageSandbox))
 	{
 		// Create Sandbox
-		SandboxPlatformFile = MakeUnique<FConcertSandboxPlatformFile>(LiveSession->GetSession().GetSessionWorkingDirectory() / TEXT("Sandbox"));
-		SandboxPlatformFile->Initialize(&FPlatformFileManager::Get().GetPlatformFile(), TEXT(""));
+		SandboxManager = FConcertPackageSandbox::Make(LiveSession->GetSession().GetSessionWorkingDirectory());
 	}
 
 	// Previously these event handlers were wrapped in the GIsEditor below.  We moved them out to support take recording
@@ -104,11 +107,11 @@ FConcertClientPackageManager::~FConcertClientPackageManager()
 	UPackage::PackageDirtyStateChangedEvent.RemoveAll(this);
 	PackageBridge->OnLocalPackageEvent().RemoveAll(this);
 
-	if (SandboxPlatformFile)
+	if (SandboxManager)
 	{
 		// Discard Sandbox and gather packages to be reloaded/purged
-		SandboxPlatformFile->DiscardSandbox(PackagesPendingHotReload, PackagesPendingPurge);
-		SandboxPlatformFile.Reset();
+		SandboxManager->DiscardSandbox(PackagesPendingHotReload, PackagesPendingPurge);
+		SandboxManager.Reset();
 	}
 	FEditorDelegates::OnDuplicateActorsBegin.RemoveAll(this);
 	FEditorDelegates::OnDuplicateActorsEnd.RemoveAll(this);
@@ -303,14 +306,14 @@ TOptional<FString> FConcertClientPackageManager::GetValidPackageSessionPath(FNam
 TOptional<FString> FConcertClientPackageManager::GetDeletedPackagePath(FName PackageName) const
 {
 #if WITH_EDITOR
-	check(SandboxPlatformFile.IsValid());
-	FConcertSandboxPlatformFile* PlatformFile = SandboxPlatformFile.Get();
-	auto ShouldPersistPackageWithExtension = [PlatformFile](const FString& PackageName, const FString& Extension) -> TOptional<FString>
+	check(SandboxManager.IsValid());
+	UE::ConcertClient::IConcertSandbox* Sandbox = SandboxManager.Get();
+	auto ShouldPersistPackageWithExtension = [Sandbox](const FString& PackageName, const FString& Extension) -> TOptional<FString>
 	{
 		FString FullPath;
 		if (FPackageName::TryConvertLongPackageNameToFilename(PackageName, FullPath, Extension))
 		{
-			if (PlatformFile->DeletedPackageExistsInNonSandbox(FullPath))
+			if (Sandbox->DeletedPackageExistsInNonSandbox(FullPath))
 			{
 				return MoveTemp(FullPath);
 			}
@@ -333,7 +336,7 @@ TOptional<FString> FConcertClientPackageManager::GetDeletedPackagePath(FName Pac
 FPersistResult FConcertClientPackageManager::PersistSessionChanges(FPersistParameters InParam)
 {
 #if WITH_EDITOR
-	if (SandboxPlatformFile)
+	if (SandboxManager)
 	{
 		// Transform all the package names into actual filenames
 		TArray<FString, TInlineAllocator<8>> FilesToPersist;
@@ -344,7 +347,7 @@ FPersistResult FConcertClientPackageManager::PersistSessionChanges(FPersistParam
 				FilesToPersist.Add(MoveTemp(ValidPath.GetValue()));
 			}
 		}
-		return SandboxPlatformFile->PersistSandbox(FilesToPersist, MoveTemp(InParam));
+		return SandboxManager->PersistSandbox(FilesToPersist, MoveTemp(InParam));
 	}
 #endif
 	return {};
diff --git a/Engine/Plugins/Developer/Concert/ConcertSync/ConcertSyncClient/Source/ConcertSyncClient/Private/ConcertClientPackageManager.h b/Engine/Plugins/Developer/Concert/ConcertSync/ConcertSyncClient/Source/ConcertSyncClient/Private/ConcertClientPackageManager.h
--- a/Engine/Plugins/Developer/Concert/ConcertSync/ConcertSyncClient/Source/ConcertSyncClient/Private/ConcertClientPackageManager.h
+++ b/Engine/Plugins/Developer/Concert/ConcertSync/ConcertSyncClient/Source/ConcertSyncClient/Private/ConcertClientPackageManager.h
@@ -7,6 +7,7 @@
 #include "IConcertClientPackageBridge.h"
 
 #include "ConcertClientPersistData.h"
+#include "Sandbox/IConcertSandbox.h"
 
 class FConcertSyncClientLiveSession;
 class FConcertSandboxPlatformFile;
@@ -184,10 +185,8 @@ private:
 	void AddPendingReloadForNewExternalMaps(const FConcertPackageInfo& PackageInfo, const FString& PackagePathname);
 
 #if WITH_EDITOR
-	/**
-	 * Sandbox for storing package changes to disk within a Concert session.
-	 */
-	TUniquePtr<FConcertSandboxPlatformFile> SandboxPlatformFile; // TODO: Will need to ensure the sandbox also works on cooked clients
+	/** Sandbox for storing package changes to disk within a Concert session. */
+	TUniquePtr<UE::ConcertClient::IConcertSandbox> SandboxManager; // TODO: Will need to ensure the sandbox also works on cooked clients
 #endif
 
 	/**
diff --git a/Engine/Plugins/Developer/Concert/ConcertSync/ConcertSyncClient/Source/ConcertSyncClient/Private/Sandbox/ConcertPackageSandbox.cpp b/Engine/Plugins/Developer/Concert/ConcertSync/ConcertSyncClient/Source/ConcertSyncClient/Private/Sandbox/ConcertPackageSandbox.cpp
new file mode 100644
--- /dev/null
+++ b/Engine/Plugins/Developer/Concert/ConcertSync/ConcertSyncClient/Source/ConcertSyncClient/Private/Sandbox/ConcertPackageSandbox.cpp
@@ -0,0 +1,100 @@
+﻿// Copyright Epic Games, Inc. All Rights Reserved.
+
+#include "ConcertPackageSandbox.h"
+
+#include "ConcertClientPersistData.h"
+#include "ConcertLogGlobal.h"
+#include "IFileSandboxCoreModule.h"
+#include "ISandboxInstance.h"
+#include "ISandboxManager.h"
+#include "Misc/PackageName.h"
+#include "Types/Manager/NewSandboxArgs.h"
+#include "Types/Sandbox/DiscardResult.h"
+
+namespace UE::ConcertClient
+{
+TUniquePtr<FConcertPackageSandbox> FConcertPackageSandbox::Make(const FString& InWorkspaceDirectory)
+{
+	using namespace FileSandboxCore;
+	ISandboxManager& SandboxManager = IFileSandboxCoreModule::Get().GetSandboxManager();
+
+	// TODO DP: We should ask the user nicely whether they are sure to leave the currently active sandbox
+	if (SandboxManager.HasActiveSandbox())
+	{
+		SandboxManager.LeaveSandbox();
+	}
+
+	FNewSandboxArgs Args(TEXT("Sandbox"));
+	Args.SandboxBasePath = InWorkspaceDirectory;
+	const FNewSandboxResult Result = SandboxManager.CreateNewSandbox(Args);
+	if (Result.HasError())
+	{
+		UE_LOG(LogConcert, Error, TEXT("Failed to enter sandbox. Reason: %s"), *LexToString(Result.GetError().Reason));
+		return nullptr;
+	}
+	return MakeUnique<FConcertPackageSandbox>();
+}
+	
+FConcertPackageSandbox::FConcertPackageSandbox()
+{
+	using namespace FileSandboxCore;
+	ISandboxManager& SandboxManager = IFileSandboxCoreModule::Get().GetSandboxManager();
+	check(SandboxManager.HasActiveSandbox()); // FConcertPackageSandbox::Make called constructor without any valid sandbox.
+
+	// Listen for external sandbox shutdown. It's not supposed to happen so we warn if it does.
+	SandboxManager.OnPreSandboxShutdown().AddRaw(this, &FConcertPackageSandbox::HandleExternalSandboxShutdown);
+}
+
+FConcertPackageSandbox::~FConcertPackageSandbox()
+{
+	using namespace FileSandboxCore;
+	ISandboxManager& SandboxManager = IFileSandboxCoreModule::Get().GetSandboxManager();
+	SandboxManager.OnPreSandboxShutdown().RemoveAll(this);
+
+	// Must shutdown after we've stopped listening for shutdown to prevent us logging an error incorrectly.
+	SandboxManager.LeaveSandbox();
+}
+
+FPersistResult FConcertPackageSandbox::PersistSandbox(TArrayView<const FString> InFiles, FPersistParameters InParams)
+{
+	using namespace FileSandboxCore;
+	
+	ISandboxManager& SandboxManager = IFileSandboxCoreModule::Get().GetSandboxManager();
+	if (ISandboxInstance* Sandbox = SandboxManager.GetActiveSandboxInstance())
+	{
+		const FPersistArgs Args { InFiles, InParams.bShouldMakeWritableIfNoSourceControl };
+		const FPersistResult Result = Sandbox->PersistSandbox(Args);
+		return ::FPersistResult(Result.FailureReasons, static_cast<::EPersistStatus>(Result.PersistStatus));
+	}
+
+	return ::FPersistResult { .PersistStatus = ::EPersistStatus::Failure };
+}
+	
+void FConcertPackageSandbox::DiscardSandbox(TArray<FName>& PackagesPendingHotReload, TArray<FName>& PackagesPendingPurge)
+{
+	using namespace FileSandboxCore;
+	ISandboxManager& SandboxManager = IFileSandboxCoreModule::Get().GetSandboxManager();
+	if (ISandboxInstance* Sandbox = SandboxManager.GetActiveSandboxInstance())
+	{
+		const FDiscardResult DiscardResult = Sandbox->DiscardSandbox();
+		PackagesPendingHotReload.Append(DiscardResult.PackagesPendingHotReload);
+		PackagesPendingPurge.Append(DiscardResult.PackagesPendingPurge);
+	}
+}
+
+bool FConcertPackageSandbox::DeletedPackageExistsInNonSandbox(FString InFilename) const
+{
+	using namespace FileSandboxCore;
+	ISandboxManager& SandboxManager = IFileSandboxCoreModule::Get().GetSandboxManager();
+	ISandboxInstance* Sandbox = SandboxManager.GetActiveSandboxInstance();
+	return Sandbox && Sandbox->DeletedPackageExistsInNonSandbox(InFilename);
+}
+
+void FConcertPackageSandbox::HandleExternalSandboxShutdown(FileSandboxCore::ISandboxInstance& SandboxInstance)
+{
+	// TODO DP: Maybe investigate API to prevent external code from shutting down our sandbox, e.g. by passing in a TAttribute<bool> with FNewSandboxResult.
+	ensureMsgf(false, TEXT("Concert sandbox was shutdown externally"));
+	UE_LOG(LogConcert, Error, TEXT("Concert sandbox was shutdown externally"));
+}
+}
+
diff --git a/Engine/Plugins/Developer/Concert/ConcertSync/ConcertSyncClient/Source/ConcertSyncClient/Private/Sandbox/ConcertPackageSandbox.h b/Engine/Plugins/Developer/Concert/ConcertSync/ConcertSyncClient/Source/ConcertSyncClient/Private/Sandbox/ConcertPackageSandbox.h
new file mode 100644
--- /dev/null
+++ b/Engine/Plugins/Developer/Concert/ConcertSync/ConcertSyncClient/Source/ConcertSyncClient/Private/Sandbox/ConcertPackageSandbox.h
@@ -0,0 +1,35 @@
+﻿// Copyright Epic Games, Inc. All Rights Reserved.
+
+#pragma once
+
+#include "IConcertSandbox.h"
+#include "Templates/UniquePtr.h"
+#include "Templates/UnrealTemplate.h"
+
+namespace UE::FileSandboxCore { class ISandboxInstance; }
+
+namespace UE::ConcertClient
+{
+/** Handles sandboxing changes made to files while in a Concert session. */
+class FConcertPackageSandbox : public FNoncopyable, public IConcertSandbox
+{
+public:
+	
+	/** Creates a file sandbox. */
+	static TUniquePtr<FConcertPackageSandbox> Make(const FString& InWorkspaceDirectory);
+	
+	FConcertPackageSandbox();
+	virtual ~FConcertPackageSandbox() override;
+
+	//~ Begin IConcertSandbox Interface
+	virtual FPersistResult PersistSandbox(TArrayView<const FString> InFiles, FPersistParameters InParams) override;
+	virtual void DiscardSandbox(TArray<FName>& PackagesPendingHotReload, TArray<FName>& PackagesPendingPurge) override;
+	virtual bool DeletedPackageExistsInNonSandbox(FString InFilename) const override;
+	//~ Begin IConcertSandbox Interface
+
+private:
+	
+	/** Handles some external sandbox shutting down our sandbox (somebody evil can just call ISandboxManager::LeaveSandbox while we're in Concert). */
+	void HandleExternalSandboxShutdown(FileSandboxCore::ISandboxInstance& SandboxInstance);
+};
+}
diff --git a/Engine/Plugins/Developer/Concert/ConcertSync/ConcertSyncClient/Source/ConcertSyncClient/Private/Sandbox/IConcertSandbox.h b/Engine/Plugins/Developer/Concert/ConcertSync/ConcertSyncClient/Source/ConcertSyncClient/Private/Sandbox/IConcertSandbox.h
new file mode 100644
--- /dev/null
+++ b/Engine/Plugins/Developer/Concert/ConcertSync/ConcertSyncClient/Source/ConcertSyncClient/Private/Sandbox/IConcertSandbox.h
@@ -0,0 +1,44 @@
+﻿// Copyright Epic Games, Inc. All Rights Reserved.
+
+#pragma once
+
+#include "Containers/ContainersFwd.h"
+#include "Templates/UniquePtr.h"
+
+class FName;
+class FString;
+struct FDateTime;
+struct FPersistParameters;
+struct FPersistResult;
+
+namespace UE::ConcertClient
+{
+class IConcertSandbox
+{
+public:
+	
+	/**
+	 *	Persist the file list from the sandbox state onto the real files Will mark files for which the operation was
+	 *	succesful as persisted.
+	 *
+	 *	@param InFiles	List of files names to persist as full paths.
+	 *	@param InParams	Parameters controlling the persist options specified by the user.
+	 *	@return Returns a persist result object that contains result status.
+	 */
+	virtual FPersistResult PersistSandbox(TArrayView<const FString> InFiles, FPersistParameters InParams) = 0;
+
+	/** 
+	 *	Discard the sandbox state
+	 *	This will trigger directory watcher notifications for files that are restored.
+	 *	This will also gather packages that need to be hot reloaded or purge from memory
+	 *	@param OutPackagePendingHotReload	Array to be filled with package name that need to be hot reloaded
+	 *	@param OutPackagesPendingPurge		Array to be filled with package name that need to be purged
+	 */
+	virtual void DiscardSandbox(TArray<FName>& OutPackagesPendingHotReload, TArray<FName>& OutPackagesPendingPurge) = 0;
+
+	/** @return true if the given package file exists on non sandbox path. */
+	virtual bool DeletedPackageExistsInNonSandbox(FString InFilename) const = 0;
+
+	virtual ~IConcertSandbox() = default;
+};
+}
diff --git a/Engine/Plugins/Developer/Sandbox/FileSandbox/Source/FileSandboxCore/Private/Types/Manager/NewSandboxError.cpp b/Engine/Plugins/Developer/Sandbox/FileSandbox/Source/FileSandboxCore/Private/Types/Manager/NewSandboxError.cpp
new file mode 100644
--- /dev/null
+++ b/Engine/Plugins/Developer/Sandbox/FileSandbox/Source/FileSandboxCore/Private/Types/Manager/NewSandboxError.cpp
@@ -0,0 +1,21 @@
+﻿// Copyright Epic Games, Inc. All Rights Reserved.
+
+#include "Types/Manager/NewSandboxError.h"
+
+#include "Containers/UnrealString.h"
+#include "Misc/AssertionMacros.h"
+
+namespace UE::FileSandboxCore
+{
+FString LexToString(ENewSandboxErrorReason InReason)
+{
+	static_assert(static_cast<int32>(ENewSandboxErrorReason::Count) == 3, "Update this switch");
+	switch (InReason)
+	{
+	case ENewSandboxErrorReason::UnsuitablePath: return TEXT("UnsuitablePath");
+	case ENewSandboxErrorReason::IOError: return TEXT("IOError");
+	case ENewSandboxErrorReason::Unspecified: return TEXT("Unspecified");
+	default: checkNoEntry(); return TEXT("Unknown");
+	}
+}
+}
diff --git a/Engine/Plugins/Developer/Sandbox/FileSandbox/Source/FileSandboxCore/Public/Types/Manager/NewSandboxError.h b/Engine/Plugins/Developer/Sandbox/FileSandbox/Source/FileSandboxCore/Public/Types/Manager/NewSandboxError.h
--- a/Engine/Plugins/Developer/Sandbox/FileSandbox/Source/FileSandboxCore/Public/Types/Manager/NewSandboxError.h
+++ b/Engine/Plugins/Developer/Sandbox/FileSandbox/Source/FileSandboxCore/Public/Types/Manager/NewSandboxError.h
@@ -4,6 +4,8 @@
 
 #include "HAL/Platform.h"
 
+class FString;
+
 namespace UE::FileSandboxCore
 {
 enum class ENewSandboxErrorReason : uint8
@@ -15,7 +17,10 @@ enum class ENewSandboxErrorReason : uint8
 	IOError,
 
 	/** Dummy error code for initializing to a deterministic value. */
-	Unspecified
+	Unspecified,
+
+	/** ADD NEW MEMBERS ABOVE. Not a real reason. */
+	Count
 };
 
 /** Explains why a new sandbox could not be created. */ 
@@ -23,4 +28,7 @@ struct FNewSandboxError
 {
 	ENewSandboxErrorReason Reason = ENewSandboxErrorReason::Unspecified;
 };
-}
\ No newline at end of file
+
+/** @return Error reason converted to a string. */
+FILESANDBOXCORE_API FString LexToString(ENewSandboxErrorReason InReason);
+}
