[Backout] - CL48389162
#rnx
Original CL Desc
-----------------------------------------------------------------
Concert: Convert ConcertClient to use new FileSandbox API for package sandbox.

#Synced-CL 47637010

#Synced-CL 47637010

diff --git a/Engine/Plugins/Developer/Concert/ConcertSync/ConcertSyncClient/ConcertSyncClient.uplugin b/Engine/Plugins/Developer/Concert/ConcertSync/ConcertSyncClient/ConcertSyncClient.uplugin
--- a/Engine/Plugins/Developer/Concert/ConcertSync/ConcertSyncClient/ConcertSyncClient.uplugin
+++ b/Engine/Plugins/Developer/Concert/ConcertSync/ConcertSyncClient/ConcertSyncClient.uplugin
@@ -31,10 +31,6 @@
 		{
 			"Name": "ConcertMain",
 			"Enabled": true
-		},
-		{
-			"Name": "FileSandbox",
-			"Enabled": true
 		}
 	]
 }
diff --git a/Engine/Plugins/Developer/Concert/ConcertSync/ConcertSyncClient/Source/ConcertSyncClient/ConcertSyncClient.Build.cs b/Engine/Plugins/Developer/Concert/ConcertSync/ConcertSyncClient/Source/ConcertSyncClient/ConcertSyncClient.Build.cs
--- a/Engine/Plugins/Developer/Concert/ConcertSync/ConcertSyncClient/Source/ConcertSyncClient/ConcertSyncClient.Build.cs
+++ b/Engine/Plugins/Developer/Concert/ConcertSync/ConcertSyncClient/Source/ConcertSyncClient/ConcertSyncClient.Build.cs
@@ -23,7 +23,6 @@ namespace UnrealBuildTool.Rules
 				{
 					"AssetRegistry",
 					"ConcertTransport",
-					"FileSandboxCore",
 					"HeadMountedDisplay",
 					"InputCore",
 					"JsonUtilities",
diff --git a/Engine/Plugins/Developer/Concert/ConcertSync/ConcertSyncClient/Source/ConcertSyncClient/Private/ConcertClientPackageManager.cpp b/Engine/Plugins/Developer/Concert/ConcertSync/ConcertSyncClient/Source/ConcertSyncClient/Private/ConcertClientPackageManager.cpp
--- a/Engine/Plugins/Developer/Concert/ConcertSync/ConcertSyncClient/Source/ConcertSyncClient/Private/ConcertClientPackageManager.cpp
+++ b/Engine/Plugins/Developer/Concert/ConcertSync/ConcertSyncClient/Source/ConcertSyncClient/Private/ConcertClientPackageManager.cpp
@@ -27,8 +27,6 @@
 #include "HAL/PlatformFileManager.h"
 
 #include "ISourceControlModule.h"
-#include "Sandbox/ConcertPackageSandbox.h"
-#include "Sandbox/IConcertSandbox.h"
 #include "UObject/UObjectHash.h"
 
 #if WITH_EDITOR
@@ -72,8 +70,6 @@ FConcertClientPackageManager::FConcertClientPackageManager(TSharedRef<FConcertSy
 	, bIgnorePackageDirtyEvent(false)
 	, FileSharingService(MoveTemp(InFileSharingService))
 {
-	using namespace UE::ConcertClient;
-	
 	check(LiveSession->IsValidSession());
 	check(EnumHasAnyFlags(LiveSession->GetSessionFlags(), EConcertSyncSessionFlags::EnablePackages));
 	check(PackageBridge);
@@ -82,7 +78,8 @@ FConcertClientPackageManager::FConcertClientPackageManager(TSharedRef<FConcertSy
 	if (EnumHasAnyFlags(LiveSession->GetSessionFlags(), EConcertSyncSessionFlags::ShouldUsePackageSandbox))
 	{
 		// Create Sandbox
-		SandboxManager = FConcertPackageSandbox::Make(LiveSession->GetSession().GetSessionWorkingDirectory());
+		SandboxPlatformFile = MakeUnique<FConcertSandboxPlatformFile>(LiveSession->GetSession().GetSessionWorkingDirectory() / TEXT("Sandbox"));
+		SandboxPlatformFile->Initialize(&FPlatformFileManager::Get().GetPlatformFile(), TEXT(""));
 	}
 
 	// Previously these event handlers were wrapped in the GIsEditor below.  We moved them out to support take recording
@@ -107,11 +104,11 @@ FConcertClientPackageManager::~FConcertClientPackageManager()
 	UPackage::PackageDirtyStateChangedEvent.RemoveAll(this);
 	PackageBridge->OnLocalPackageEvent().RemoveAll(this);
 
-	if (SandboxManager)
+	if (SandboxPlatformFile)
 	{
 		// Discard Sandbox and gather packages to be reloaded/purged
-		SandboxManager->DiscardSandbox(PackagesPendingHotReload, PackagesPendingPurge);
-		SandboxManager.Reset();
+		SandboxPlatformFile->DiscardSandbox(PackagesPendingHotReload, PackagesPendingPurge);
+		SandboxPlatformFile.Reset();
 	}
 	FEditorDelegates::OnDuplicateActorsBegin.RemoveAll(this);
 	FEditorDelegates::OnDuplicateActorsEnd.RemoveAll(this);
@@ -306,14 +303,14 @@ TOptional<FString> FConcertClientPackageManager::GetValidPackageSessionPath(FNam
 TOptional<FString> FConcertClientPackageManager::GetDeletedPackagePath(FName PackageName) const
 {
 #if WITH_EDITOR
-	check(SandboxManager.IsValid());
-	UE::ConcertClient::IConcertSandbox* Sandbox = SandboxManager.Get();
-	auto ShouldPersistPackageWithExtension = [Sandbox](const FString& PackageName, const FString& Extension) -> TOptional<FString>
+	check(SandboxPlatformFile.IsValid());
+	FConcertSandboxPlatformFile* PlatformFile = SandboxPlatformFile.Get();
+	auto ShouldPersistPackageWithExtension = [PlatformFile](const FString& PackageName, const FString& Extension) -> TOptional<FString>
 	{
 		FString FullPath;
 		if (FPackageName::TryConvertLongPackageNameToFilename(PackageName, FullPath, Extension))
 		{
-			if (Sandbox->DeletedPackageExistsInNonSandbox(FullPath))
+			if (PlatformFile->DeletedPackageExistsInNonSandbox(FullPath))
 			{
 				return MoveTemp(FullPath);
 			}
@@ -336,7 +333,7 @@ TOptional<FString> FConcertClientPackageManager::GetDeletedPackagePath(FName Pac
 FPersistResult FConcertClientPackageManager::PersistSessionChanges(FPersistParameters InParam)
 {
 #if WITH_EDITOR
-	if (SandboxManager)
+	if (SandboxPlatformFile)
 	{
 		// Transform all the package names into actual filenames
 		TArray<FString, TInlineAllocator<8>> FilesToPersist;
@@ -347,7 +344,7 @@ FPersistResult FConcertClientPackageManager::PersistSessionChanges(FPersistParam
 				FilesToPersist.Add(MoveTemp(ValidPath.GetValue()));
 			}
 		}
-		return SandboxManager->PersistSandbox(FilesToPersist, MoveTemp(InParam));
+		return SandboxPlatformFile->PersistSandbox(FilesToPersist, MoveTemp(InParam));
 	}
 #endif
 	return {};
diff --git a/Engine/Plugins/Developer/Concert/ConcertSync/ConcertSyncClient/Source/ConcertSyncClient/Private/ConcertClientPackageManager.h b/Engine/Plugins/Developer/Concert/ConcertSync/ConcertSyncClient/Source/ConcertSyncClient/Private/ConcertClientPackageManager.h
--- a/Engine/Plugins/Developer/Concert/ConcertSync/ConcertSyncClient/Source/ConcertSyncClient/Private/ConcertClientPackageManager.h
+++ b/Engine/Plugins/Developer/Concert/ConcertSync/ConcertSyncClient/Source/ConcertSyncClient/Private/ConcertClientPackageManager.h
@@ -7,7 +7,6 @@
 #include "IConcertClientPackageBridge.h"
 
 #include "ConcertClientPersistData.h"
-#include "Sandbox/IConcertSandbox.h"
 
 class FConcertSyncClientLiveSession;
 class FConcertSandboxPlatformFile;
@@ -185,8 +184,10 @@ private:
 	void AddPendingReloadForNewExternalMaps(const FConcertPackageInfo& PackageInfo, const FString& PackagePathname);
 
 #if WITH_EDITOR
-	/** Sandbox for storing package changes to disk within a Concert session. */
-	TUniquePtr<UE::ConcertClient::IConcertSandbox> SandboxManager; // TODO: Will need to ensure the sandbox also works on cooked clients
+	/**
+	 * Sandbox for storing package changes to disk within a Concert session.
+	 */
+	TUniquePtr<FConcertSandboxPlatformFile> SandboxPlatformFile; // TODO: Will need to ensure the sandbox also works on cooked clients
 #endif
 
 	/**
diff --git a/Engine/Plugins/Developer/Concert/ConcertSync/ConcertSyncClient/Source/ConcertSyncClient/Private/Sandbox/ConcertPackageSandbox.cpp b/Engine/Plugins/Developer/Concert/ConcertSync/ConcertSyncClient/Source/ConcertSyncClient/Private/Sandbox/ConcertPackageSandbox.cpp
deleted file mode 100644
diff --git a/Engine/Plugins/Developer/Concert/ConcertSync/ConcertSyncClient/Source/ConcertSyncClient/Private/Sandbox/ConcertPackageSandbox.h b/Engine/Plugins/Developer/Concert/ConcertSync/ConcertSyncClient/Source/ConcertSyncClient/Private/Sandbox/ConcertPackageSandbox.h
deleted file mode 100644
diff --git a/Engine/Plugins/Developer/Concert/ConcertSync/ConcertSyncClient/Source/ConcertSyncClient/Private/Sandbox/IConcertSandbox.h b/Engine/Plugins/Developer/Concert/ConcertSync/ConcertSyncClient/Source/ConcertSyncClient/Private/Sandbox/IConcertSandbox.h
deleted file mode 100644
diff --git a/Engine/Plugins/Developer/Sandbox/FileSandbox/Source/FileSandboxCore/Private/Types/Manager/NewSandboxError.cpp b/Engine/Plugins/Developer/Sandbox/FileSandbox/Source/FileSandboxCore/Private/Types/Manager/NewSandboxError.cpp
deleted file mode 100644
diff --git a/Engine/Plugins/Developer/Sandbox/FileSandbox/Source/FileSandboxCore/Public/Types/Manager/NewSandboxError.h b/Engine/Plugins/Developer/Sandbox/FileSandbox/Source/FileSandboxCore/Public/Types/Manager/NewSandboxError.h
--- a/Engine/Plugins/Developer/Sandbox/FileSandbox/Source/FileSandboxCore/Public/Types/Manager/NewSandboxError.h
+++ b/Engine/Plugins/Developer/Sandbox/FileSandbox/Source/FileSandboxCore/Public/Types/Manager/NewSandboxError.h
@@ -4,8 +4,6 @@
 
 #include "HAL/Platform.h"
 
-class FString;
-
 namespace UE::FileSandboxCore
 {
 enum class ENewSandboxErrorReason : uint8
@@ -17,10 +15,7 @@ enum class ENewSandboxErrorReason : uint8
 	IOError,
 
 	/** Dummy error code for initializing to a deterministic value. */
-	Unspecified,
-
-	/** ADD NEW MEMBERS ABOVE. Not a real reason. */
-	Count
+	Unspecified
 };
 
 /** Explains why a new sandbox could not be created. */ 
@@ -28,7 +23,4 @@ struct FNewSandboxError
 {
 	ENewSandboxErrorReason Reason = ENewSandboxErrorReason::Unspecified;
 };
-
-/** @return Error reason converted to a string. */
-FILESANDBOXCORE_API FString LexToString(ENewSandboxErrorReason InReason);
-}
+}
\ No newline at end of file
