Concert: Convert ConcertClient to use new FileSandbox API for package sandbox.

Changes:
- The package manager now uses FileSandbox API
- Remove source control proxy: FileSandboxCore now handles overriding the source control state.
- Specifically, ISourceControlState::IsLocal is now mirrored by FSandboxSourceControlStateProxy. This ensures redirectors are created.
- Remove FPersistParameters::SourceControlProvider

#Synced-CL 47637010

diff --git a/Engine/Plugins/Developer/Concert/ConcertApp/MultiUserClient/Source/MultiUserClient/Private/ConcertWorkspaceUI.cpp b/Engine/Plugins/Developer/Concert/ConcertApp/MultiUserClient/Source/MultiUserClient/Private/ConcertWorkspaceUI.cpp
--- a/Engine/Plugins/Developer/Concert/ConcertApp/MultiUserClient/Source/MultiUserClient/Private/ConcertWorkspaceUI.cpp
+++ b/Engine/Plugins/Developer/Concert/ConcertApp/MultiUserClient/Source/MultiUserClient/Private/ConcertWorkspaceUI.cpp
@@ -726,7 +726,7 @@ bool FConcertWorkspaceUI::PromptPersistSessionChanges()
 	FText NotificationSub;
 
 	FPersistResult PersistResult = ClientWorkspacePin->PersistSessionChanges(
-		{PersistCmd.PackagesToPersist, &ISourceControlModule::Get().GetProvider(),PersistCmd.bShouldMakeFilesWritable});
+		{PersistCmd.PackagesToPersist, PersistCmd.bShouldMakeFilesWritable});
 
 	bool bSuccess = PersistResult.PersistStatus == EPersistStatus::Success;
 	if (bSuccess)
diff --git a/Engine/Plugins/Developer/Concert/ConcertSync/ConcertSyncClient/ConcertSyncClient.uplugin b/Engine/Plugins/Developer/Concert/ConcertSync/ConcertSyncClient/ConcertSyncClient.uplugin
--- a/Engine/Plugins/Developer/Concert/ConcertSync/ConcertSyncClient/ConcertSyncClient.uplugin
+++ b/Engine/Plugins/Developer/Concert/ConcertSync/ConcertSyncClient/ConcertSyncClient.uplugin
@@ -31,6 +31,10 @@
 		{
 			"Name": "ConcertMain",
 			"Enabled": true
+		},
+		{
+			"Name": "FileSandbox",
+			"Enabled": true
 		}
 	]
 }
diff --git a/Engine/Plugins/Developer/Concert/ConcertSync/ConcertSyncClient/Source/ConcertSyncClient/ConcertSyncClient.Build.cs b/Engine/Plugins/Developer/Concert/ConcertSync/ConcertSyncClient/Source/ConcertSyncClient/ConcertSyncClient.Build.cs
--- a/Engine/Plugins/Developer/Concert/ConcertSync/ConcertSyncClient/Source/ConcertSyncClient/ConcertSyncClient.Build.cs
+++ b/Engine/Plugins/Developer/Concert/ConcertSync/ConcertSyncClient/Source/ConcertSyncClient/ConcertSyncClient.Build.cs
@@ -23,6 +23,7 @@ namespace UnrealBuildTool.Rules
 				{
 					"AssetRegistry",
 					"ConcertTransport",
+					"FileSandboxCore",
 					"HeadMountedDisplay",
 					"InputCore",
 					"JsonUtilities",
diff --git a/Engine/Plugins/Developer/Concert/ConcertSync/ConcertSyncClient/Source/ConcertSyncClient/Private/ConcertClientPackageManager.cpp b/Engine/Plugins/Developer/Concert/ConcertSync/ConcertSyncClient/Source/ConcertSyncClient/Private/ConcertClientPackageManager.cpp
--- a/Engine/Plugins/Developer/Concert/ConcertSync/ConcertSyncClient/Source/ConcertSyncClient/Private/ConcertClientPackageManager.cpp
+++ b/Engine/Plugins/Developer/Concert/ConcertSync/ConcertSyncClient/Source/ConcertSyncClient/Private/ConcertClientPackageManager.cpp
@@ -10,7 +10,6 @@
 #include "ConcertLogGlobal.h"
 #include "ConcertWorkspaceData.h"
 #include "ConcertWorkspaceMessages.h"
-#include "ConcertSandboxPlatformFile.h"
 #include "ConcertSyncClientUtil.h"
 #include "ConcertSyncSettings.h"
 #include "ConcertUtil.h"
@@ -27,6 +26,8 @@
 #include "HAL/PlatformFileManager.h"
 
 #include "ISourceControlModule.h"
+#include "Sandbox/ConcertPackageSandbox.h"
+#include "Sandbox/IConcertSandbox.h"
 #include "UObject/UObjectHash.h"
 
 #if WITH_EDITOR
@@ -70,6 +71,8 @@ FConcertClientPackageManager::FConcertClientPackageManager(TSharedRef<FConcertSy
 	, bIgnorePackageDirtyEvent(false)
 	, FileSharingService(MoveTemp(InFileSharingService))
 {
+	using namespace UE::ConcertClient;
+	
 	check(LiveSession->IsValidSession());
 	check(EnumHasAnyFlags(LiveSession->GetSessionFlags(), EConcertSyncSessionFlags::EnablePackages));
 	check(PackageBridge);
@@ -78,8 +81,7 @@ FConcertClientPackageManager::FConcertClientPackageManager(TSharedRef<FConcertSy
 	if (EnumHasAnyFlags(LiveSession->GetSessionFlags(), EConcertSyncSessionFlags::ShouldUsePackageSandbox))
 	{
 		// Create Sandbox
-		SandboxPlatformFile = MakeUnique<FConcertSandboxPlatformFile>(LiveSession->GetSession().GetSessionWorkingDirectory() / TEXT("Sandbox"));
-		SandboxPlatformFile->Initialize(&FPlatformFileManager::Get().GetPlatformFile(), TEXT(""));
+		SandboxManager = FConcertPackageSandbox::CreateOrLoad(LiveSession->GetSession().GetSessionWorkingDirectory());
 	}
 
 	// Previously these event handlers were wrapped in the GIsEditor below.  We moved them out to support take recording
@@ -104,11 +106,11 @@ FConcertClientPackageManager::~FConcertClientPackageManager()
 	UPackage::PackageDirtyStateChangedEvent.RemoveAll(this);
 	PackageBridge->OnLocalPackageEvent().RemoveAll(this);
 
-	if (SandboxPlatformFile)
+	if (SandboxManager)
 	{
 		// Discard Sandbox and gather packages to be reloaded/purged
-		SandboxPlatformFile->DiscardSandbox(PackagesPendingHotReload, PackagesPendingPurge);
-		SandboxPlatformFile.Reset();
+		SandboxManager->DiscardSandbox(PackagesPendingHotReload, PackagesPendingPurge);
+		SandboxManager.Reset();
 	}
 	FEditorDelegates::OnDuplicateActorsBegin.RemoveAll(this);
 	FEditorDelegates::OnDuplicateActorsEnd.RemoveAll(this);
@@ -303,14 +305,14 @@ TOptional<FString> FConcertClientPackageManager::GetValidPackageSessionPath(FNam
 TOptional<FString> FConcertClientPackageManager::GetDeletedPackagePath(FName PackageName) const
 {
 #if WITH_EDITOR
-	check(SandboxPlatformFile.IsValid());
-	FConcertSandboxPlatformFile* PlatformFile = SandboxPlatformFile.Get();
-	auto ShouldPersistPackageWithExtension = [PlatformFile](const FString& PackageName, const FString& Extension) -> TOptional<FString>
+	check(SandboxManager.IsValid());
+	UE::ConcertClient::IConcertSandbox* Sandbox = SandboxManager.Get();
+	auto ShouldPersistPackageWithExtension = [Sandbox](const FString& PackageName, const FString& Extension) -> TOptional<FString>
 	{
 		FString FullPath;
 		if (FPackageName::TryConvertLongPackageNameToFilename(PackageName, FullPath, Extension))
 		{
-			if (PlatformFile->DeletedPackageExistsInNonSandbox(FullPath))
+			if (Sandbox->DeletedPackageExistsInNonSandbox(FullPath))
 			{
 				return MoveTemp(FullPath);
 			}
@@ -333,7 +335,7 @@ TOptional<FString> FConcertClientPackageManager::GetDeletedPackagePath(FName Pac
 FPersistResult FConcertClientPackageManager::PersistSessionChanges(FPersistParameters InParam)
 {
 #if WITH_EDITOR
-	if (SandboxPlatformFile)
+	if (SandboxManager)
 	{
 		// Transform all the package names into actual filenames
 		TArray<FString, TInlineAllocator<8>> FilesToPersist;
@@ -344,7 +346,7 @@ FPersistResult FConcertClientPackageManager::PersistSessionChanges(FPersistParam
 				FilesToPersist.Add(MoveTemp(ValidPath.GetValue()));
 			}
 		}
-		return SandboxPlatformFile->PersistSandbox(FilesToPersist, MoveTemp(InParam));
+		return SandboxManager->PersistSandbox(FilesToPersist, MoveTemp(InParam));
 	}
 #endif
 	return {};
@@ -542,7 +544,7 @@ void FConcertClientPackageManager::HandleLocalPackageEvent(const FConcertPackage
 	// if the package data has been filtered out of the session persist it immediately from the sandbox
 	else
 	{
-		PersistSessionChanges({MakeArrayView(&PackageInfo.PackageName, 1), &ISourceControlModule::Get().GetProvider()});
+		PersistSessionChanges({MakeArrayView(&PackageInfo.PackageName, 1) });
 	}
 }
 
diff --git a/Engine/Plugins/Developer/Concert/ConcertSync/ConcertSyncClient/Source/ConcertSyncClient/Private/ConcertClientPackageManager.h b/Engine/Plugins/Developer/Concert/ConcertSync/ConcertSyncClient/Source/ConcertSyncClient/Private/ConcertClientPackageManager.h
--- a/Engine/Plugins/Developer/Concert/ConcertSync/ConcertSyncClient/Source/ConcertSyncClient/Private/ConcertClientPackageManager.h
+++ b/Engine/Plugins/Developer/Concert/ConcertSync/ConcertSyncClient/Source/ConcertSyncClient/Private/ConcertClientPackageManager.h
@@ -7,6 +7,7 @@
 #include "IConcertClientPackageBridge.h"
 
 #include "ConcertClientPersistData.h"
+#include "Sandbox/IConcertSandbox.h"
 
 class FConcertSyncClientLiveSession;
 class FConcertSandboxPlatformFile;
@@ -184,10 +185,8 @@ private:
 	void AddPendingReloadForNewExternalMaps(const FConcertPackageInfo& PackageInfo, const FString& PackagePathname);
 
 #if WITH_EDITOR
-	/**
-	 * Sandbox for storing package changes to disk within a Concert session.
-	 */
-	TUniquePtr<FConcertSandboxPlatformFile> SandboxPlatformFile; // TODO: Will need to ensure the sandbox also works on cooked clients
+	/** Sandbox for storing package changes to disk within a Concert session. */
+	TUniquePtr<UE::ConcertClient::IConcertSandbox> SandboxManager; // TODO UE-354294: Will need to ensure the sandbox also works on cooked clients
 #endif
 
 	/**
diff --git a/Engine/Plugins/Developer/Concert/ConcertSync/ConcertSyncClient/Source/ConcertSyncClient/Private/ConcertSandboxPlatformFile.cpp b/Engine/Plugins/Developer/Concert/ConcertSync/ConcertSyncClient/Source/ConcertSyncClient/Private/ConcertSandboxPlatformFile.cpp
deleted file mode 100644
diff --git a/Engine/Plugins/Developer/Concert/ConcertSync/ConcertSyncClient/Source/ConcertSyncClient/Private/ConcertSandboxPlatformFile.h b/Engine/Plugins/Developer/Concert/ConcertSync/ConcertSyncClient/Source/ConcertSyncClient/Private/ConcertSandboxPlatformFile.h
deleted file mode 100644
diff --git a/Engine/Plugins/Developer/Concert/ConcertSync/ConcertSyncClient/Source/ConcertSyncClient/Private/ConcertSourceControlProxy.cpp b/Engine/Plugins/Developer/Concert/ConcertSync/ConcertSyncClient/Source/ConcertSyncClient/Private/ConcertSourceControlProxy.cpp
deleted file mode 100644
diff --git a/Engine/Plugins/Developer/Concert/ConcertSync/ConcertSyncClient/Source/ConcertSyncClient/Private/ConcertSourceControlProxy.h b/Engine/Plugins/Developer/Concert/ConcertSync/ConcertSyncClient/Source/ConcertSyncClient/Private/ConcertSourceControlProxy.h
deleted file mode 100644
diff --git a/Engine/Plugins/Developer/Concert/ConcertSync/ConcertSyncClient/Source/ConcertSyncClient/Private/ConcertSyncClient.cpp b/Engine/Plugins/Developer/Concert/ConcertSync/ConcertSyncClient/Source/ConcertSyncClient/Private/ConcertSyncClient.cpp
--- a/Engine/Plugins/Developer/Concert/ConcertSync/ConcertSyncClient/Source/ConcertSyncClient/Private/ConcertSyncClient.cpp
+++ b/Engine/Plugins/Developer/Concert/ConcertSync/ConcertSyncClient/Source/ConcertSyncClient/Private/ConcertSyncClient.cpp
@@ -10,7 +10,6 @@
 #include "ConcertClientWorkspace.h"
 #include "ConcertClientSequencerManager.h"
 #include "ConcertClientPresenceManager.h"
-#include "ConcertSourceControlProxy.h"
 #include "Replication/Manager/ReplicationManager.h"
 
 LLM_DECLARE_TAG(Concert_ConcertSyncClient);
@@ -25,7 +24,6 @@ FConcertSyncClient::FConcertSyncClient(
 	, Bridges(InBridges)
 #if WITH_EDITOR
 	, SequencerManager(MakeUnique<FConcertClientSequencerManager>(this))
-	, SourceControlProxy(MakeUnique<FConcertSourceControlProxy>())
 #endif
 {
 	check(Bridges.IsValid());
@@ -121,7 +119,7 @@ void FConcertSyncClient::PersistSpecificChanges(TArrayView<const FName> InPackag
 #if WITH_EDITOR
 	if (Workspace)
 	{
-		Workspace->PersistSessionChanges({InPackages, SourceControlProxy.Get()});
+		Workspace->PersistSessionChanges({ InPackages });
 	}
 #endif
 }
@@ -132,7 +130,7 @@ void FConcertSyncClient::PersistAllSessionChanges()
 	if (Workspace)
 	{
 		TArray<FName> SessionChanges = Workspace->GatherSessionChanges();
-		Workspace->PersistSessionChanges({SessionChanges, SourceControlProxy.Get()});
+		Workspace->PersistSessionChanges({ SessionChanges });
 	}
 #endif
 }
@@ -159,23 +157,12 @@ void FConcertSyncClient::CreateWorkspace(const TSharedRef<FConcertSyncClientLive
 	Workspace = MakeShared<FConcertClientWorkspace>(UE::ConcertSyncClient::FSessionBindArgs{ InLiveSession, Bridges } , FileSharingService, this);
 	OnWorkspaceStartupDelegate.Broadcast(Workspace);
 #if WITH_EDITOR
-	if (GIsEditor && EnumHasAllFlags(SessionFlags, EConcertSyncSessionFlags::EnablePackages | EConcertSyncSessionFlags::ShouldUsePackageSandbox))
-	{
-		// TODO: Revisit this, as all it seems to be used for now is forcing redirectors to be left behind
-		SourceControlProxy->SetWorkspace(Workspace);
-	}
 	SequencerManager->SetActiveWorkspace(Workspace);
 #endif
 }
 
 void FConcertSyncClient::DestroyWorkspace()
 {
-#if WITH_EDITOR
-	if (GIsEditor)
-	{
-		SourceControlProxy->SetWorkspace(nullptr);
-	}
-#endif
 	OnWorkspaceShutdownDelegate.Broadcast(Workspace);
 	Workspace.Reset();
 }
diff --git a/Engine/Plugins/Developer/Concert/ConcertSync/ConcertSyncClient/Source/ConcertSyncClient/Private/ConcertSyncClient.h b/Engine/Plugins/Developer/Concert/ConcertSync/ConcertSyncClient/Source/ConcertSyncClient/Private/ConcertSyncClient.h
--- a/Engine/Plugins/Developer/Concert/ConcertSync/ConcertSyncClient/Source/ConcertSyncClient/Private/ConcertSyncClient.h
+++ b/Engine/Plugins/Developer/Concert/ConcertSync/ConcertSyncClient/Source/ConcertSyncClient/Private/ConcertSyncClient.h
@@ -99,8 +99,5 @@ private:
 
 	/** Sequencer event manager for Concert session. */
 	TUniquePtr<FConcertClientSequencerManager> SequencerManager;
-
-	/** Source Control Provider Proxy for Concert session. */
-	TUniquePtr<FConcertSourceControlProxy> SourceControlProxy;
 #endif
 };
diff --git a/Engine/Plugins/Developer/Concert/ConcertSync/ConcertSyncClient/Source/ConcertSyncClient/Private/Sandbox/ConcertPackageSandbox.cpp b/Engine/Plugins/Developer/Concert/ConcertSync/ConcertSyncClient/Source/ConcertSyncClient/Private/Sandbox/ConcertPackageSandbox.cpp
new file mode 100644
--- /dev/null
+++ b/Engine/Plugins/Developer/Concert/ConcertSync/ConcertSyncClient/Source/ConcertSyncClient/Private/Sandbox/ConcertPackageSandbox.cpp
@@ -0,0 +1,149 @@
+﻿// Copyright Epic Games, Inc. All Rights Reserved.
+
+#include "ConcertPackageSandbox.h"
+
+#include "ConcertClientPersistData.h"
+#include "ConcertLogGlobal.h"
+#include "IFileSandboxCoreModule.h"
+#include "ISandboxInstance.h"
+#include "ISandboxManager.h"
+#include "Misc/PackageName.h"
+#include "Types/Manager/LoadSandboxByNameArgs.h"
+#include "Types/Manager/NewSandboxArgs.h"
+#include "Types/Manager/SandboxCreationResult.h"
+#include "Types/Sandbox/DiscardResult.h"
+
+namespace UE::ConcertClient
+{
+namespace Private
+{
+static bool CreateNewSandbox(
+	FileSandboxCore::ISandboxManager& InSandboxManager, const FileSandboxCore::FSandboxInitArgs& InInitArgs,
+	const FString& InSandboxName, const FString& InWorkspaceDirectory
+	)
+{
+	using namespace FileSandboxCore;
+	
+	FNewSandboxArgs Args(InSandboxName);
+	Args.SandboxBasePath = InWorkspaceDirectory;
+	Args.InitArgs = InInitArgs;
+	
+	const FNewSandboxResult Result = InSandboxManager.CreateNewSandbox(Args);
+	
+	const bool bIsSuccess = Result.HasValue();
+	UE_CLOG(!bIsSuccess, LogConcert, Error, TEXT("Failed to create sandbox. Reason: %s"), *LexToString(Result.GetError().Reason));
+	return bIsSuccess;
+}
+	
+static bool LoadSandbox(
+	FileSandboxCore::ISandboxManager& InSandboxManager, const FileSandboxCore::FSandboxInitArgs& InInitArgs, 
+	const FString& InSandboxName, const FString& InWorkspaceDirectory
+	)
+{
+	using namespace FileSandboxCore;
+	
+	FLoadSandboxByNameArgs Args(InSandboxName, InWorkspaceDirectory);
+	Args.InitArgs = InInitArgs;
+	
+	const FLoadSandboxResult Result = InSandboxManager.LoadNamedSandbox(Args);
+	
+	const bool bIsSuccess = Result.HasValue();
+	UE_CLOG(!bIsSuccess, LogConcert, Error, TEXT("Failed to load sandbox. Reason: %s"), *LexToString(Result.GetError().Reason));
+	return bIsSuccess;
+}
+}
+	
+TUniquePtr<FConcertPackageSandbox> FConcertPackageSandbox::CreateOrLoad(const FString& InWorkspaceDirectory)
+{
+	using namespace FileSandboxCore;
+	ISandboxManager& SandboxManager = IFileSandboxCoreModule::Get().GetSandboxManager();
+
+	// TODO DP: We should ask the user nicely whether they are sure to leave the currently active sandbox
+	if (SandboxManager.HasActiveSandbox())
+	{
+		SandboxManager.LeaveSandbox();
+	}
+
+	const FString SandboxName(TEXT("Sandbox"));
+	const TOptional<FString> ExistingDirectory = FindSandboxByName(SandboxName, InWorkspaceDirectory);
+	
+	// This causes ISourceControlState::IsLocal to return false.
+	// Concert live propagates assets to other clients, so no assets can be local
+	// This is used to determine whether redirector assets are left when things are renamed
+	const FSandboxInitArgs InitArgs { .Flags = ESandboxInitFlags::ForceRedirectorsOnRenameInSourceControl };
+	
+	if (ExistingDirectory && Private::LoadSandbox(SandboxManager, InitArgs, SandboxName, InWorkspaceDirectory))
+	{
+		return MakeUnique<FConcertPackageSandbox>();
+	}
+	
+	if (Private::CreateNewSandbox(SandboxManager, InitArgs, SandboxName, InWorkspaceDirectory))
+	{
+		return MakeUnique<FConcertPackageSandbox>();
+	}
+	
+	return nullptr;
+}
+	
+FConcertPackageSandbox::FConcertPackageSandbox()
+{
+	using namespace FileSandboxCore;
+	ISandboxManager& SandboxManager = IFileSandboxCoreModule::Get().GetSandboxManager();
+	check(SandboxManager.HasActiveSandbox()); // FConcertPackageSandbox::Make called constructor without any valid sandbox.
+
+	// Listen for external sandbox shutdown. It's not supposed to happen so we warn if it does.
+	SandboxManager.OnPreSandboxShutdown().AddRaw(this, &FConcertPackageSandbox::HandleExternalSandboxShutdown);
+}
+
+FConcertPackageSandbox::~FConcertPackageSandbox()
+{
+	using namespace FileSandboxCore;
+	ISandboxManager& SandboxManager = IFileSandboxCoreModule::Get().GetSandboxManager();
+	SandboxManager.OnPreSandboxShutdown().RemoveAll(this);
+
+	// Must shutdown after we've stopped listening for shutdown to prevent us logging an error incorrectly.
+	SandboxManager.LeaveSandbox();
+}
+
+FPersistResult FConcertPackageSandbox::PersistSandbox(TArrayView<const FString> InFiles, FPersistParameters InParams)
+{
+	using namespace FileSandboxCore;
+	
+	ISandboxManager& SandboxManager = IFileSandboxCoreModule::Get().GetSandboxManager();
+	if (ISandboxInstance* Sandbox = SandboxManager.GetActiveSandboxInstance())
+	{
+		const FPersistArgs Args { InFiles, InParams.bShouldMakeWritableIfNoSourceControl };
+		const FPersistResult Result = Sandbox->PersistSandbox(Args);
+		return ::FPersistResult(Result.FailureReasons, static_cast<::EPersistStatus>(Result.PersistStatus));
+	}
+
+	return ::FPersistResult { .PersistStatus = ::EPersistStatus::Failure };
+}
+	
+void FConcertPackageSandbox::DiscardSandbox(TArray<FName>& PackagesPendingHotReload, TArray<FName>& PackagesPendingPurge)
+{
+	using namespace FileSandboxCore;
+	ISandboxManager& SandboxManager = IFileSandboxCoreModule::Get().GetSandboxManager();
+	if (ISandboxInstance* Sandbox = SandboxManager.GetActiveSandboxInstance())
+	{
+		const FDiscardResult DiscardResult = Sandbox->DiscardSandbox();
+		PackagesPendingHotReload.Append(DiscardResult.PackagesPendingHotReload);
+		PackagesPendingPurge.Append(DiscardResult.PackagesPendingPurge);
+	}
+}
+
+bool FConcertPackageSandbox::DeletedPackageExistsInNonSandbox(FString InFilename) const
+{
+	using namespace FileSandboxCore;
+	ISandboxManager& SandboxManager = IFileSandboxCoreModule::Get().GetSandboxManager();
+	ISandboxInstance* Sandbox = SandboxManager.GetActiveSandboxInstance();
+	return Sandbox && Sandbox->DeletedPackageExistsInNonSandbox(InFilename);
+}
+
+void FConcertPackageSandbox::HandleExternalSandboxShutdown(FileSandboxCore::ISandboxInstance& SandboxInstance)
+{
+	// TODO DP: Maybe investigate API to prevent external code from shutting down our sandbox, e.g. by passing in a TAttribute<bool> with FNewSandboxResult.
+	ensureMsgf(false, TEXT("Concert sandbox was shutdown externally"));
+	UE_LOG(LogConcert, Error, TEXT("Concert sandbox was shutdown externally"));
+}
+}
\ No newline at end of file
diff --git a/Engine/Plugins/Developer/Concert/ConcertSync/ConcertSyncClient/Source/ConcertSyncClient/Private/Sandbox/ConcertPackageSandbox.h b/Engine/Plugins/Developer/Concert/ConcertSync/ConcertSyncClient/Source/ConcertSyncClient/Private/Sandbox/ConcertPackageSandbox.h
new file mode 100644
--- /dev/null
+++ b/Engine/Plugins/Developer/Concert/ConcertSync/ConcertSyncClient/Source/ConcertSyncClient/Private/Sandbox/ConcertPackageSandbox.h
@@ -0,0 +1,35 @@
+﻿// Copyright Epic Games, Inc. All Rights Reserved.
+
+#pragma once
+
+#include "IConcertSandbox.h"
+#include "Templates/UniquePtr.h"
+#include "Templates/UnrealTemplate.h"
+
+namespace UE::FileSandboxCore { class ISandboxInstance; }
+
+namespace UE::ConcertClient
+{
+/** Handles sandboxing changes made to files while in a Concert session. */
+class FConcertPackageSandbox : public FNoncopyable, public IConcertSandbox
+{
+public:
+	
+	/** Creates a file sandbox. */
+	static TUniquePtr<FConcertPackageSandbox> CreateOrLoad(const FString& InWorkspaceDirectory);
+	
+	FConcertPackageSandbox();
+	virtual ~FConcertPackageSandbox() override;
+
+	//~ Begin IConcertSandbox Interface
+	virtual FPersistResult PersistSandbox(TArrayView<const FString> InFiles, FPersistParameters InParams) override;
+	virtual void DiscardSandbox(TArray<FName>& PackagesPendingHotReload, TArray<FName>& PackagesPendingPurge) override;
+	virtual bool DeletedPackageExistsInNonSandbox(FString InFilename) const override;
+	//~ Begin IConcertSandbox Interface
+
+private:
+	
+	/** Handles some external sandbox shutting down our sandbox (somebody evil can just call ISandboxManager::LeaveSandbox while we're in Concert). */
+	void HandleExternalSandboxShutdown(FileSandboxCore::ISandboxInstance& SandboxInstance);
+};
+}
\ No newline at end of file
diff --git a/Engine/Plugins/Developer/Concert/ConcertSync/ConcertSyncClient/Source/ConcertSyncClient/Private/Sandbox/IConcertSandbox.h b/Engine/Plugins/Developer/Concert/ConcertSync/ConcertSyncClient/Source/ConcertSyncClient/Private/Sandbox/IConcertSandbox.h
new file mode 100644
--- /dev/null
+++ b/Engine/Plugins/Developer/Concert/ConcertSync/ConcertSyncClient/Source/ConcertSyncClient/Private/Sandbox/IConcertSandbox.h
@@ -0,0 +1,44 @@
+﻿// Copyright Epic Games, Inc. All Rights Reserved.
+
+#pragma once
+
+#include "Containers/ContainersFwd.h"
+#include "Templates/UniquePtr.h"
+
+class FName;
+class FString;
+struct FDateTime;
+struct FPersistParameters;
+struct FPersistResult;
+
+namespace UE::ConcertClient
+{
+class IConcertSandbox
+{
+public:
+	
+	/**
+	 *	Persist the file list from the sandbox state onto the real files Will mark files for which the operation was
+	 *	succesful as persisted.
+	 *
+	 *	@param InFiles	List of files names to persist as full paths.
+	 *	@param InParams	Parameters controlling the persist options specified by the user.
+	 *	@return Returns a persist result object that contains result status.
+	 */
+	virtual FPersistResult PersistSandbox(TArrayView<const FString> InFiles, FPersistParameters InParams) = 0;
+
+	/** 
+	 *	Discard the sandbox state
+	 *	This will trigger directory watcher notifications for files that are restored.
+	 *	This will also gather packages that need to be hot reloaded or purge from memory
+	 *	@param OutPackagePendingHotReload	Array to be filled with package name that need to be hot reloaded
+	 *	@param OutPackagesPendingPurge		Array to be filled with package name that need to be purged
+	 */
+	virtual void DiscardSandbox(TArray<FName>& OutPackagesPendingHotReload, TArray<FName>& OutPackagesPendingPurge) = 0;
+
+	/** @return true if the given package file exists on non sandbox path. */
+	virtual bool DeletedPackageExistsInNonSandbox(FString InFilename) const = 0;
+
+	virtual ~IConcertSandbox() = default;
+};
+}
diff --git a/Engine/Plugins/Developer/Concert/ConcertSync/ConcertSyncClient/Source/ConcertSyncClient/Public/ConcertClientPersistData.h b/Engine/Plugins/Developer/Concert/ConcertSync/ConcertSyncClient/Source/ConcertSyncClient/Public/ConcertClientPersistData.h
--- a/Engine/Plugins/Developer/Concert/ConcertSync/ConcertSyncClient/Source/ConcertSyncClient/Public/ConcertClientPersistData.h
+++ b/Engine/Plugins/Developer/Concert/ConcertSync/ConcertSyncClient/Source/ConcertSyncClient/Public/ConcertClientPersistData.h
@@ -14,9 +14,6 @@ struct FPersistParameters
 	/** Name of the packages to persist.*/
 	TArrayView<const FName> PackagesToPersist;
 
-	/** Pointer to the source control provider. */
-	ISourceControlProvider* SourceControlProvider;
-
 	/** Boolean to control if we will force the file writable when no source control provider is active. */
 	bool bShouldMakeWritableIfNoSourceControl = false;
 };
diff --git a/Engine/Plugins/Developer/Sandbox/FileSandbox/FileSandbox.uplugin b/Engine/Plugins/Developer/Sandbox/FileSandbox/FileSandbox.uplugin
--- a/Engine/Plugins/Developer/Sandbox/FileSandbox/FileSandbox.uplugin
+++ b/Engine/Plugins/Developer/Sandbox/FileSandbox/FileSandbox.uplugin
@@ -15,13 +15,11 @@
 	"Hidden": true,
 	"IsBetaVersion": true,
 	"Installed": false,
-	"SupportedPrograms": [ "UnrealMultiUserServer", "CoopMultiUserServer", "UnrealMultiUserSlateServer", "UnrealRecoverySvc", "CrashReportClientEditor", "LiveLinkHub" ],
 	"Modules": [
 		{
 			"Name": "FileSandboxCore",
-			"Type": "Runtime",
-			"LoadingPhase": "Default", 
-			"ProgramAllowList": [ "UnrealMultiUserServer", "CoopMultiUserServer", "UnrealMultiUserSlateServer", "UnrealRecoverySvc", "CrashReportClientEditor", "LiveLinkHub" ]
+			"Type": "UncookedOnly",
+			"LoadingPhase": "Default"
 		}
 	]
 }
