Concert: Prevent sandbox from being shutdown while in Concert sesssion

#Synced-CL 47637010

diff --git a/Engine/Plugins/Developer/Concert/ConcertSync/ConcertSyncClient/Source/ConcertSyncClient/Private/ConcertClientPackageManager.cpp b/Engine/Plugins/Developer/Concert/ConcertSync/ConcertSyncClient/Source/ConcertSyncClient/Private/ConcertClientPackageManager.cpp
--- a/Engine/Plugins/Developer/Concert/ConcertSync/ConcertSyncClient/Source/ConcertSyncClient/Private/ConcertClientPackageManager.cpp
+++ b/Engine/Plugins/Developer/Concert/ConcertSync/ConcertSyncClient/Source/ConcertSyncClient/Private/ConcertClientPackageManager.cpp
@@ -65,13 +65,19 @@ bool RunPackageFilters(const TArray<FPackageClassFilter>& InFilters, const FConc
 }
 } // end namespace ConcertClientPackageManagerUtil
 
-FConcertClientPackageManager::FConcertClientPackageManager(TSharedRef<FConcertSyncClientLiveSession> InLiveSession, IConcertClientPackageBridge* InPackageBridge, TSharedPtr<IConcertFileSharingService> InFileSharingService)
-	: LiveSession(MoveTemp(InLiveSession))
+FConcertClientPackageManager::FConcertClientPackageManager(
+	const TSharedRef<FConcertSyncClientLiveSession>& InLiveSession,
+	IConcertClientPackageBridge* InPackageBridge, 
+	TSharedPtr<IConcertFileSharingService> InFileSharingService,
+	const FString& InClientRole
+	)
+	: LiveSession(InLiveSession)
 	, PackageBridge(InPackageBridge)
 	, bIgnorePackageDirtyEvent(false)
 	, FileSharingService(MoveTemp(InFileSharingService))
+
 {
-	using namespace UE::ConcertClient;
+	using namespace UE::ConcertSyncClient;
 	
 	check(LiveSession->IsValidSession());
 	check(EnumHasAnyFlags(LiveSession->GetSessionFlags(), EConcertSyncSessionFlags::EnablePackages));
@@ -81,7 +87,7 @@ FConcertClientPackageManager::FConcertClientPackageManager(TSharedRef<FConcertSy
 	if (EnumHasAnyFlags(LiveSession->GetSessionFlags(), EConcertSyncSessionFlags::ShouldUsePackageSandbox))
 	{
 		// Create Sandbox
-		SandboxManager = FConcertPackageSandbox::CreateOrLoad(LiveSession->GetSession().GetSessionWorkingDirectory());
+		SandboxManager = FConcertPackageSandbox::CreateOrLoad(LiveSession.ToSharedRef(), InClientRole);
 	}
 
 	// Previously these event handlers were wrapped in the GIsEditor below.  We moved them out to support take recording
@@ -306,7 +312,7 @@ TOptional<FString> FConcertClientPackageManager::GetDeletedPackagePath(FName Pac
 {
 #if WITH_EDITOR
 	check(SandboxManager.IsValid());
-	UE::ConcertClient::IConcertSandbox* Sandbox = SandboxManager.Get();
+	UE::ConcertSyncClient::IConcertSandbox* Sandbox = SandboxManager.Get();
 	auto ShouldPersistPackageWithExtension = [Sandbox](const FString& PackageName, const FString& Extension) -> TOptional<FString>
 	{
 		FString FullPath;
diff --git a/Engine/Plugins/Developer/Concert/ConcertSync/ConcertSyncClient/Source/ConcertSyncClient/Private/ConcertClientPackageManager.h b/Engine/Plugins/Developer/Concert/ConcertSync/ConcertSyncClient/Source/ConcertSyncClient/Private/ConcertClientPackageManager.h
--- a/Engine/Plugins/Developer/Concert/ConcertSync/ConcertSyncClient/Source/ConcertSyncClient/Private/ConcertClientPackageManager.h
+++ b/Engine/Plugins/Developer/Concert/ConcertSync/ConcertSyncClient/Source/ConcertSyncClient/Private/ConcertClientPackageManager.h
@@ -7,13 +7,18 @@
 #include "IConcertClientPackageBridge.h"
 
 #include "ConcertClientPersistData.h"
+#include "IConcertClientModule.h"
+#include "Misc/Attribute.h"
 #include "Sandbox/IConcertSandbox.h"
 
+
+class FConcertSyncClient;
 class FConcertSyncClientLiveSession;
 class FConcertSandboxPlatformFile;
 class ISourceControlProvider;
 class IConcertFileSharingService;
 class UPackage;
+namespace UE::ConcertSyncClient { class FConcertSandboxLock; }
 
 struct FAssetData;
 struct FConcertPackage;
@@ -34,8 +39,14 @@ public:
 	 * @param InLiveSession The session for which the package manager will be used.
 	 * @param InPackageBridge The package bridge (used by the package manager to register itself for package events)
 	 * @param InFileSharingService Optional service (can be null) used by the client to exchange very large packages with the server.
-	 */
-	FConcertClientPackageManager(TSharedRef<FConcertSyncClientLiveSession> InLiveSession, IConcertClientPackageBridge* InPackageBridge, TSharedPtr<IConcertFileSharingService> InFileSharingService);
+	 * @param InClientRole Role of the client. Used by sandbox system to generate lock reason text.
+	 */
+	FConcertClientPackageManager(
+		const TSharedRef<FConcertSyncClientLiveSession>& InLiveSession, 
+		IConcertClientPackageBridge* InPackageBridge, 
+		TSharedPtr<IConcertFileSharingService> InFileSharingService,
+		const FString& InClientRole
+		);
 	~FConcertClientPackageManager();
 	
 	/**
@@ -186,7 +197,7 @@ private:
 
 #if WITH_EDITOR
 	/** Sandbox for storing package changes to disk within a Concert session. */
-	TUniquePtr<UE::ConcertClient::IConcertSandbox> SandboxManager; // TODO UE-354294: Will need to ensure the sandbox also works on cooked clients
+	TUniquePtr<UE::ConcertSyncClient::IConcertSandbox> SandboxManager; // TODO UE-354294: Will need to ensure the sandbox also works on cooked clients
 #endif
 
 	/**
diff --git a/Engine/Plugins/Developer/Concert/ConcertSync/ConcertSyncClient/Source/ConcertSyncClient/Private/ConcertClientWorkspace.cpp b/Engine/Plugins/Developer/Concert/ConcertSync/ConcertSyncClient/Source/ConcertSyncClient/Private/ConcertClientWorkspace.cpp
--- a/Engine/Plugins/Developer/Concert/ConcertSync/ConcertSyncClient/Source/ConcertSyncClient/Private/ConcertClientWorkspace.cpp
+++ b/Engine/Plugins/Developer/Concert/ConcertSync/ConcertSyncClient/Source/ConcertSyncClient/Private/ConcertClientWorkspace.cpp
@@ -537,7 +537,9 @@ void FConcertClientWorkspace::BindSession(const UE::ConcertSyncClient::FSessionB
 	// Create Package Manager
 	if (EnumHasAnyFlags(LiveSession->GetSessionFlags(), EConcertSyncSessionFlags::EnablePackages))
 	{
-		PackageManager = MakeUnique<FConcertClientPackageManager>(LiveSession.ToSharedRef(), SessionBindArgs.Bridges.PackageBridge, FileSharingService);
+		PackageManager = MakeUnique<FConcertClientPackageManager>(
+			LiveSession.ToSharedRef(), SessionBindArgs.Bridges.PackageBridge, FileSharingService, OwnerSyncClient->GetConcertClient()->GetRole()
+			);
 	}
 
 	// Create Lock Manager
diff --git a/Engine/Plugins/Developer/Concert/ConcertSync/ConcertSyncClient/Source/ConcertSyncClient/Private/Sandbox/ConcertPackageSandbox.cpp b/Engine/Plugins/Developer/Concert/ConcertSync/ConcertSyncClient/Source/ConcertSyncClient/Private/Sandbox/ConcertPackageSandbox.cpp
--- a/Engine/Plugins/Developer/Concert/ConcertSync/ConcertSyncClient/Source/ConcertSyncClient/Private/Sandbox/ConcertPackageSandbox.cpp
+++ b/Engine/Plugins/Developer/Concert/ConcertSync/ConcertSyncClient/Source/ConcertSyncClient/Private/Sandbox/ConcertPackageSandbox.cpp
@@ -4,16 +4,19 @@
 
 #include "ConcertClientPersistData.h"
 #include "ConcertLogGlobal.h"
+#include "ConcertSession.h"
+#include "ConcertSyncClientLiveSession.h"
 #include "IFileSandboxCoreModule.h"
 #include "ISandboxInstance.h"
 #include "ISandboxManager.h"
 #include "Misc/PackageName.h"
+#include "Sandbox/ConcertSandboxLock.h"
 #include "Types/Manager/LoadSandboxByNameArgs.h"
 #include "Types/Manager/NewSandboxArgs.h"
 #include "Types/Manager/SandboxCreationResult.h"
 #include "Types/Sandbox/DiscardResult.h"
 
-namespace UE::ConcertClient
+namespace UE::ConcertSyncClient
 {
 namespace Private
 {
@@ -53,7 +56,10 @@ static bool LoadSandbox(
 }
 }
 	
-TUniquePtr<FConcertPackageSandbox> FConcertPackageSandbox::CreateOrLoad(const FString& InWorkspaceDirectory)
+TUniquePtr<FConcertPackageSandbox> FConcertPackageSandbox::CreateOrLoad(
+	const TSharedRef<FConcertSyncClientLiveSession>& InLiveSession,
+	const FString& InClientRole
+	)
 {
 	using namespace FileSandboxCore;
 	ISandboxManager& SandboxManager = IFileSandboxCoreModule::Get().GetSandboxManager();
@@ -65,44 +71,40 @@ TUniquePtr<FConcertPackageSandbox> FConcertPackageSandbox::CreateOrLoad(const FS
 	}
 
 	const FString SandboxName(TEXT("Sandbox"));
-	const TOptional<FString> ExistingDirectory = FindSandboxByName(SandboxName, InWorkspaceDirectory);
+	const FString WorkspaceDirectory = InLiveSession->GetSession().GetSessionWorkingDirectory();
+	const TOptional<FString> ExistingDirectory = FindSandboxByName(SandboxName, WorkspaceDirectory);
 	
-	// This causes ISourceControlState::IsLocal to return false.
-	// Concert live propagates assets to other clients, so no assets can be local
-	// This is used to determine whether redirector assets are left when things are renamed
-	const FSandboxInitArgs InitArgs { .Flags = ESandboxInitFlags::ForceRedirectorsOnRenameInSourceControl };
+	const TSharedRef<FConcertSandboxLock> Lock = MakeShared<FConcertSandboxLock>(InClientRole);
+	const FSandboxInitArgs InitArgs
+	{ 
+		// This causes ISourceControlState::IsLocal to return false.
+		// Concert live propagates assets to other clients, so no assets can be local
+		// This is used to determine whether redirector assets are left when things are renamed
+		.Flags = ESandboxInitFlags::ForceRedirectorsOnRenameInSourceControl,
+		.Lock = Lock
+	};
 	
-	if (ExistingDirectory && Private::LoadSandbox(SandboxManager, InitArgs, SandboxName, InWorkspaceDirectory))
+	if (ExistingDirectory && Private::LoadSandbox(SandboxManager, InitArgs, SandboxName, WorkspaceDirectory))
 	{
-		return MakeUnique<FConcertPackageSandbox>();
+		return MakeUnique<FConcertPackageSandbox>(Lock);
 	}
 	
-	if (Private::CreateNewSandbox(SandboxManager, InitArgs, SandboxName, InWorkspaceDirectory))
+	if (Private::CreateNewSandbox(SandboxManager, InitArgs, SandboxName, WorkspaceDirectory))
 	{
-		return MakeUnique<FConcertPackageSandbox>();
+		return MakeUnique<FConcertPackageSandbox>(Lock);
 	}
 	
 	return nullptr;
 }
-	
-FConcertPackageSandbox::FConcertPackageSandbox()
-{
-	using namespace FileSandboxCore;
-	ISandboxManager& SandboxManager = IFileSandboxCoreModule::Get().GetSandboxManager();
-	check(SandboxManager.HasActiveSandbox()); // FConcertPackageSandbox::Make called constructor without any valid sandbox.
 
-	// Listen for external sandbox shutdown. It's not supposed to happen so we warn if it does.
-	SandboxManager.OnPreSandboxShutdown().AddRaw(this, &FConcertPackageSandbox::HandleExternalSandboxShutdown);
-}
+FConcertPackageSandbox::FConcertPackageSandbox(const TSharedRef<FConcertSandboxLock> InLock)
+	: Lock(InLock)
+{}
 
 FConcertPackageSandbox::~FConcertPackageSandbox()
 {
-	using namespace FileSandboxCore;
-	ISandboxManager& SandboxManager = IFileSandboxCoreModule::Get().GetSandboxManager();
-	SandboxManager.OnPreSandboxShutdown().RemoveAll(this);
-
-	// Must shutdown after we've stopped listening for shutdown to prevent us logging an error incorrectly.
-	SandboxManager.LeaveSandbox();
+	Lock->ReleaseLock();
+	FileSandboxCore::IFileSandboxCoreModule::Get().GetSandboxManager().LeaveSandbox();
 }
 
 FPersistResult FConcertPackageSandbox::PersistSandbox(TArrayView<const FString> InFiles, FPersistParameters InParams)
@@ -139,11 +141,4 @@ bool FConcertPackageSandbox::DeletedPackageExistsInNonSandbox(FString InFilename
 	ISandboxInstance* Sandbox = SandboxManager.GetActiveSandboxInstance();
 	return Sandbox && Sandbox->DeletedPackageExistsInNonSandbox(InFilename);
 }
-
-void FConcertPackageSandbox::HandleExternalSandboxShutdown(FileSandboxCore::ISandboxInstance& SandboxInstance)
-{
-	// TODO DP: Maybe investigate API to prevent external code from shutting down our sandbox, e.g. by passing in a TAttribute<bool> with FNewSandboxResult.
-	ensureMsgf(false, TEXT("Concert sandbox was shutdown externally"));
-	UE_LOG(LogConcert, Error, TEXT("Concert sandbox was shutdown externally"));
-}
 }
\ No newline at end of file
diff --git a/Engine/Plugins/Developer/Concert/ConcertSync/ConcertSyncClient/Source/ConcertSyncClient/Private/Sandbox/ConcertPackageSandbox.h b/Engine/Plugins/Developer/Concert/ConcertSync/ConcertSyncClient/Source/ConcertSyncClient/Private/Sandbox/ConcertPackageSandbox.h
--- a/Engine/Plugins/Developer/Concert/ConcertSync/ConcertSyncClient/Source/ConcertSyncClient/Private/Sandbox/ConcertPackageSandbox.h
+++ b/Engine/Plugins/Developer/Concert/ConcertSync/ConcertSyncClient/Source/ConcertSyncClient/Private/Sandbox/ConcertPackageSandbox.h
@@ -3,22 +3,29 @@
 #pragma once
 
 #include "IConcertSandbox.h"
+#include "Templates/SharedPointer.h"
 #include "Templates/UniquePtr.h"
 #include "Templates/UnrealTemplate.h"
 
+class FConcertSyncClientLiveSession;
+
 namespace UE::FileSandboxCore { class ISandboxInstance; }
 
-namespace UE::ConcertClient
+namespace UE::ConcertSyncClient
 {
+class FConcertSandboxLock;
+
 /** Handles sandboxing changes made to files while in a Concert session. */
 class FConcertPackageSandbox : public FNoncopyable, public IConcertSandbox
 {
 public:
 	
 	/** Creates a file sandbox. */
-	static TUniquePtr<FConcertPackageSandbox> CreateOrLoad(const FString& InWorkspaceDirectory);
+	static TUniquePtr<FConcertPackageSandbox> CreateOrLoad(
+		const TSharedRef<FConcertSyncClientLiveSession>& InLiveSession, const FString& InClientRole
+		);
 	
-	FConcertPackageSandbox();
+	explicit FConcertPackageSandbox(const TSharedRef<FConcertSandboxLock> InLock);
 	virtual ~FConcertPackageSandbox() override;
 
 	//~ Begin IConcertSandbox Interface
@@ -26,10 +33,10 @@ public:
 	virtual void DiscardSandbox(TArray<FName>& PackagesPendingHotReload, TArray<FName>& PackagesPendingPurge) override;
 	virtual bool DeletedPackageExistsInNonSandbox(FString InFilename) const override;
 	//~ Begin IConcertSandbox Interface
-
+	
 private:
 	
-	/** Handles some external sandbox shutting down our sandbox (somebody evil can just call ISandboxManager::LeaveSandbox while we're in Concert). */
-	void HandleExternalSandboxShutdown(FileSandboxCore::ISandboxInstance& SandboxInstance);
+	/** Interacts with the sandbox system to prevent sandbox from being left while in Concert session. */
+	const TSharedRef<FConcertSandboxLock> Lock;
 };
 }
\ No newline at end of file
diff --git a/Engine/Plugins/Developer/Concert/ConcertSync/ConcertSyncClient/Source/ConcertSyncClient/Private/Sandbox/ConcertSandboxLock.cpp b/Engine/Plugins/Developer/Concert/ConcertSync/ConcertSyncClient/Source/ConcertSyncClient/Private/Sandbox/ConcertSandboxLock.cpp
new file mode 100644
--- /dev/null
+++ b/Engine/Plugins/Developer/Concert/ConcertSync/ConcertSyncClient/Source/ConcertSyncClient/Private/Sandbox/ConcertSandboxLock.cpp
@@ -0,0 +1,59 @@
+﻿// Copyright Epic Games, Inc. All Rights Reserved.
+
+#include "ConcertSandboxLock.h"
+
+#include "ConcertSyncClientLiveSession.h"
+#include "ConcertClient/Private/ConcertClientSession.h"
+
+namespace UE::ConcertSyncClient
+{
+namespace LockDetail
+{
+static FText RoleToReason(const FString& InString)
+{
+	// TODO: Display label for the client should be promoted to a higher level concept and should not be implemented here.
+	if (InString.Equals(TEXT("MultiUser")))
+	{
+		return NSLOCTEXT("FConcertSandboxLock", "LockReason", "The active Multi-User session is sandboxing the engine.");
+	}
+	
+	return NSLOCTEXT("FConcertSandboxLock", "LockReason.Concert", "The active Concert session is sandboxing the engine.");
+}
+}
+
+FConcertSandboxLock::FConcertSandboxLock(const FString& InClientRole)
+	: ClientRole(InClientRole)
+{}
+
+void FConcertSandboxLock::ReleaseLock()
+{
+	check(bIsLocked);
+	bIsLocked = false;
+}
+
+#define SET_REASON(Reason) if (OutReason) { *OutReason = Reason; }
+bool FConcertSandboxLock::CanLeaveSandbox(FText* OutReason) const
+{
+	if (bIsLocked)
+	{
+		SET_REASON(LockDetail::RoleToReason(ClientRole));
+		return false;
+	}
+	
+	return true;
+}
+
+bool FConcertSandboxLock::RequestLiftLock(FText* OutReason)
+{
+	// TODO UE-356773: Unsupported for now. 
+	// We need some mechanism to tell Multi-User to go through the workflow of leaving a session, which causes the persist changes UI to be shown.
+	return false;
+}
+
+bool FConcertSandboxLock::SupportsLiftingLock() const
+{
+	// TODO UE-356773: Unsupported for now. 
+	// We need some mechanism to tell Multi-User to go through the workflow of leaving a session, which causes the persist changes UI to be shown.
+	return false;
+}
+}
diff --git a/Engine/Plugins/Developer/Concert/ConcertSync/ConcertSyncClient/Source/ConcertSyncClient/Private/Sandbox/ConcertSandboxLock.h b/Engine/Plugins/Developer/Concert/ConcertSync/ConcertSyncClient/Source/ConcertSyncClient/Private/Sandbox/ConcertSandboxLock.h
new file mode 100644
--- /dev/null
+++ b/Engine/Plugins/Developer/Concert/ConcertSync/ConcertSyncClient/Source/ConcertSyncClient/Private/Sandbox/ConcertSandboxLock.h
@@ -0,0 +1,38 @@
+﻿// Copyright Epic Games, Inc. All Rights Reserved.
+
+#pragma once
+
+#include "Interface/ISandboxLock.h"
+#include "Templates/SharedPointer.h"
+
+class FConcertSyncClientLiveSession;
+class FConcertSyncClient;
+
+namespace UE::ConcertSyncClient
+{
+/** Prevents the sandbox system from leaving the current sandbox for as long as Concert is in a session. */
+class FConcertSandboxLock : public FileSandboxCore::ISandboxLock
+{
+public:
+	
+	explicit FConcertSandboxLock(const FString& InClientRole);
+	
+	/** Lifts the lock. Called when the session stops. */
+	void ReleaseLock();
+	
+	//~ Begin ISandboxLock Interface
+	virtual bool CanLeaveSandbox(FText* OutReason) const override;
+	virtual bool RequestLiftLock(FText* OutReason) override;
+	virtual bool SupportsLiftingLock() const override;
+	//~ End ISandboxLock Interface
+	
+private:
+	
+	/** The client role. Relevant for lock reason generation. */
+	const FString ClientRole;
+	
+	/** Whether the sandbox is currently locked. */
+	bool bIsLocked = true;
+};
+}
+
diff --git a/Engine/Plugins/Developer/Concert/ConcertSync/ConcertSyncClient/Source/ConcertSyncClient/Private/Sandbox/IConcertSandbox.h b/Engine/Plugins/Developer/Concert/ConcertSync/ConcertSyncClient/Source/ConcertSyncClient/Private/Sandbox/IConcertSandbox.h
--- a/Engine/Plugins/Developer/Concert/ConcertSync/ConcertSyncClient/Source/ConcertSyncClient/Private/Sandbox/IConcertSandbox.h
+++ b/Engine/Plugins/Developer/Concert/ConcertSync/ConcertSyncClient/Source/ConcertSyncClient/Private/Sandbox/IConcertSandbox.h
@@ -11,7 +11,7 @@ struct FDateTime;
 struct FPersistParameters;
 struct FPersistResult;
 
-namespace UE::ConcertClient
+namespace UE::ConcertSyncClient
 {
 class IConcertSandbox
 {
