Concert Replication: Fix removing items from TMap while iterating through it. This fixes several unit tests and CIS.

The problem was that ExtractReplicationDataForObject would be invoked indirectly and remove from PendingObjects while iterating through it. The fix is to pass a callback (IReplicationDataItem) to the callback (the TFunctionRef passed to ConsumePendingObjects). The IReplicationDataItem will correctly handle calling RemoveCurrent on the active iterator instance.

#Synced-CL 49159615

diff --git a/Engine/Plugins/Developer/Concert/ConcertSync/ConcertSyncClient/Source/ConcertSyncClient/Private/Replication/Processing/ClientReplicationDataCollector.cpp b/Engine/Plugins/Developer/Concert/ConcertSync/ConcertSyncClient/Source/ConcertSyncClient/Private/Replication/Processing/ClientReplicationDataCollector.cpp
--- a/Engine/Plugins/Developer/Concert/ConcertSync/ConcertSyncClient/Source/ConcertSyncClient/Private/Replication/Processing/ClientReplicationDataCollector.cpp
+++ b/Engine/Plugins/Developer/Concert/ConcertSync/ConcertSyncClient/Source/ConcertSyncClient/Private/Replication/Processing/ClientReplicationDataCollector.cpp
@@ -2,11 +2,12 @@
 
 #include "ClientReplicationDataCollector.h"
 
-#include "Replication/Data/ReplicationStream.h"
 #include "Replication/Data/ObjectIds.h"
+#include "Replication/Data/ReplicationStream.h"
 #include "Replication/Formats/IObjectReplicationFormat.h"
 #include "Replication/IConcertClientReplicationBridge.h"
 #include "Replication/Manager/Utils/LocalSyncControl.h"
+#include "Replication/Processing/IReplicationDataItem.h"
 #include "Replication/ReplicationPropertyFilter.h"
 
 #include "Algo/RemoveIf.h"
@@ -14,6 +15,63 @@
 
 namespace UE::ConcertSyncClient::Replication
 {
+namespace DataCollector
+{
+class FDataCollectorItem : public ConcertSyncCore::IReplicationDataItem
+{
+	FObjectInfo& ObjectInfo;
+	const FConcertReplicatedObjectId& ObjectToProcess;
+	
+	IConcertClientReplicationBridge& Bridge;
+	const FLocalSyncControl& SyncControl;
+	ConcertSyncCore::IObjectReplicationFormat& ReplicationFormat;
+	
+public:
+	
+	explicit FDataCollectorItem(
+		FObjectInfo& ObjectInfo UE_LIFETIMEBOUND, 
+		const FConcertReplicatedObjectId& ObjectToProcess UE_LIFETIMEBOUND, 
+		IConcertClientReplicationBridge& Bridge UE_LIFETIMEBOUND,
+		const FLocalSyncControl& SyncControl UE_LIFETIMEBOUND, 
+		ConcertSyncCore::IObjectReplicationFormat& ReplicationFormat UE_LIFETIMEBOUND
+		)
+		: ObjectInfo(ObjectInfo)
+		, ObjectToProcess(ObjectToProcess)
+		, Bridge(Bridge)
+		, SyncControl(SyncControl)
+		, ReplicationFormat(ReplicationFormat)
+	{}
+
+	virtual void ExtractReplicationDataForObject(
+		TFunctionRef<void(const FConcertSessionSerializedPayload& Payload)> ProcessCopyable,
+		TFunctionRef<void(FConcertSessionSerializedPayload&& Payload)> ProcessMoveable
+		) override
+	{
+		UObject* Object = Bridge.FindObjectIfAvailable(ObjectToProcess.Object);
+		// Ask the bridge to resolve the object for us. The bridge has the object cached and handles the object getting renamed, etc.
+		// This should resolve. If it does not: same logic as above - either this call is invalid or ConsumePendingObjects lied.
+		if (!ensure(Object) || !ensure(SyncControl.IsObjectAllowed(ObjectToProcess)))
+		{
+			return;
+		}
+
+		ConcertSyncCore::FReplicationPropertyFilter Filter(ObjectInfo.SelectedProperties);
+		TOptional<FConcertSessionSerializedPayload> Payload = ReplicationFormat.CreateReplicationEvent(
+			*Object,
+			[&Filter](const FArchiveSerializedPropertyChain* Chain, const FProperty& Property)
+			{
+				return Filter.ShouldSerializeProperty(Chain, Property);
+			}
+		);
+		if (Payload)
+		{
+			ProcessMoveable(MoveTemp(*Payload));
+			++ObjectInfo.ReplicationSequenceId;
+		}
+	}
+};
+}
+
 	FClientReplicationDataCollector::FClientReplicationDataCollector(
 		IConcertClientReplicationBridge& InReplicationBridge,
 		ConcertSyncCore::IObjectReplicationFormat& InReplicationFormat,
@@ -183,11 +241,13 @@ namespace UE::ConcertSyncClient::Replication
 		return Result;
 	}
 
-	void FClientReplicationDataCollector::ForEachPendingObject(TFunctionRef<void(const ConcertSyncCore::FPendingObjectReplicationInfo&)> ProcessItemFunc) const
+	void FClientReplicationDataCollector::ConsumePendingObjects(
+		TFunctionRef<void(const ConcertSyncCore::FPendingObjectReplicationInfo& Info, ConcertSyncCore::IReplicationDataItem& Item)> InProcess
+		)
 	{
-		for (const TPair<FSoftObjectPath, TArray<FObjectInfo>>& Pair : ObjectsToReplicate)
+		for (TPair<FSoftObjectPath, TArray<FObjectInfo>>& Pair : ObjectsToReplicate)
 		{
-			for (const FObjectInfo& ObjectInfo : Pair.Value)
+			for (FObjectInfo& ObjectInfo : Pair.Value)
 			{
 				const FConcertReplicatedObjectId ObjectId{
 					FConcertObjectInStreamID{ ObjectInfo.StreamId, Pair.Key },
@@ -195,15 +255,15 @@ namespace UE::ConcertSyncClient::Replication
 				};
 				if (!SyncControl.IsObjectAllowed(ObjectId))
 				{
-					continue;;
+					continue;
 				}
-				
+					
 				// The bridge has the object cached. We do not cache the object ourselves!
 				// Funky Unreal flows can cause the object to be renamed out from under us and replaced by a different instance (just by using UObject::Rename()).
 				// The bridge is aware of these flows and FindObjectIfAvailable will catch them.
 				// If FindObjectIfAvailable fails to resolve, OnObjectHidden() is triggered if the object was previously visible.
 				const UObject* Object = Bridge.FindObjectIfAvailable(Pair.Key);
-				
+					
 				const FSoftObjectPath ObjectPath = Object;
 				if (Object && ensureMsgf(ObjectPath == Object, TEXT("Sanity check: the bridge gave us an object with a different path!")))
 				{
@@ -212,56 +272,12 @@ namespace UE::ConcertSyncClient::Replication
 						FConcertReplicatedObjectId {FConcertObjectInStreamID{ ObjectInfo.StreamId, Object },ClientId },
 						ObjectInfo.ReplicationSequenceId
 					};
-					ProcessItemFunc(Info);
-				}
-			}
-		}
-	}
-
-	bool FClientReplicationDataCollector::ExtractReplicationDataForObject(
-		const FConcertReplicatedObjectId& ObjectToProcess,
-		TFunctionRef<void(const FConcertSessionSerializedPayload& Payload)> ProcessCopyable,
-		TFunctionRef<void(FConcertSessionSerializedPayload&& Payload)> ProcessMoveable
-		)
-	{
-		TArray<FObjectInfo>* ObjectInfos = ObjectsToReplicate.Find(ObjectToProcess.Object);
-		// ExtractReplicationDataForObject is supposed to be called in response to ForEachPendingObject... so either the call was invalid or ForEachPendingObject lied
-		if (!ensure(ObjectInfos))
-		{
-			return false;
-		}
-
-		// The properties to be replicated depend on the stream: search for it.
-		const int32 StreamIndex = ObjectInfos->IndexOfByPredicate([&ObjectToProcess](const FObjectInfo& ObjectInfo){ return ObjectInfo.StreamId == ObjectToProcess.StreamId; });
-		// Same logic as above: if the stream is not found, then either the call is invalid or ForEachPendingObject lied
-		if (!ensure(ObjectInfos->IsValidIndex(StreamIndex)))
-		{
-			return false;
-		}
 
-		UObject* Object = Bridge.FindObjectIfAvailable(ObjectToProcess.Object);
-		// Finally... ask the bridge to resolve the object for us. The bridge has the object cached and handles the object getting renamed, etc.
-		// This should resolve. If it does not: same logic as above - either this call is invalid or ForEachPendingObject lied.
-		if (!ensure(Object) || !ensure(SyncControl.IsObjectAllowed(ObjectToProcess)))
-		{
-			return false;
-		}
-
-		FObjectInfo& ObjectInfo = (*ObjectInfos)[StreamIndex];
-		ConcertSyncCore::FReplicationPropertyFilter Filter(ObjectInfo.SelectedProperties);
-		TOptional<FConcertSessionSerializedPayload> Payload = ReplicationFormat.CreateReplicationEvent(
-			*Object,
-			[&Filter](const FArchiveSerializedPropertyChain* Chain, const FProperty& Property)
-			{
-				return Filter.ShouldSerializeProperty(Chain, Property);
+					DataCollector::FDataCollectorItem Item(ObjectInfo, Info.ObjectId, Bridge, SyncControl, ReplicationFormat);
+					InProcess(Info, Item);
+				}
 			}
-		);
-		if (Payload)
-		{
-			ProcessMoveable(MoveTemp(*Payload));
-			++ObjectInfo.ReplicationSequenceId;
 		}
-		return true;
 	}
 
 	void FClientReplicationDataCollector::StartTrackingObject(UObject& Object)
diff --git a/Engine/Plugins/Developer/Concert/ConcertSync/ConcertSyncClient/Source/ConcertSyncClient/Private/Replication/Processing/ClientReplicationDataCollector.h b/Engine/Plugins/Developer/Concert/ConcertSync/ConcertSyncClient/Source/ConcertSyncClient/Private/Replication/Processing/ClientReplicationDataCollector.h
--- a/Engine/Plugins/Developer/Concert/ConcertSync/ConcertSyncClient/Source/ConcertSyncClient/Private/Replication/Processing/ClientReplicationDataCollector.h
+++ b/Engine/Plugins/Developer/Concert/ConcertSync/ConcertSyncClient/Source/ConcertSyncClient/Private/Replication/Processing/ClientReplicationDataCollector.h
@@ -2,6 +2,7 @@
 
 #pragma once
 
+#include "ObjectInfo.h"
 #include "Replication/Processing/IReplicationDataSource.h"
 
 #include "Containers/ArrayView.h"
@@ -87,9 +88,7 @@ namespace UE::ConcertSyncClient::Replication
 		TArray<FGuid> GetStreamsOwningObject(const FSoftObjectPath& ObjectPath) const;
 
 		//~ Begin IReplicationDataSource Interface
-		virtual void ForEachPendingObject(TFunctionRef<void(const ConcertSyncCore::FPendingObjectReplicationInfo&)> ProcessItemFunc) const override;
-		virtual int32 NumObjects() const override { return NumTrackedObjects; }
-		virtual bool ExtractReplicationDataForObject(const FConcertReplicatedObjectId& Object, TFunctionRef<void(const FConcertSessionSerializedPayload& Payload)> ProcessCopyable, TFunctionRef<void(FConcertSessionSerializedPayload&& Payload)> ProcessMoveable) override;
+		virtual void ConsumePendingObjects(TFunctionRef<void(const ConcertSyncCore::FPendingObjectReplicationInfo& Info, ConcertSyncCore::IReplicationDataItem& Item)> InProcess) override;
 		//~ End IReplicationDataSource Interface
 
 	private:
@@ -107,17 +106,6 @@ namespace UE::ConcertSyncClient::Replication
 
 		/** Endpoint ID of the client */
 		const FGuid ClientId;
-
-		struct FObjectInfo
-		{
-			/** The replication stream producing this object's data */
-			FGuid StreamId;
-			/** The properties to replicate */
-			FConcertPropertySelection SelectedProperties;
-			
-			/** Incremented every time replication data is sent out. Used for performance tracing. */
-			ConcertSyncCore::FSequenceId ReplicationSequenceId = 0;
-		};
 		
 		/** The objects and their properties to replicate */
 		TMap<FSoftObjectPath, TArray<FObjectInfo>> ObjectsToReplicate;
diff --git a/Engine/Plugins/Developer/Concert/ConcertSync/ConcertSyncClient/Source/ConcertSyncClient/Private/Replication/Processing/ObjectInfo.h b/Engine/Plugins/Developer/Concert/ConcertSync/ConcertSyncClient/Source/ConcertSyncClient/Private/Replication/Processing/ObjectInfo.h
new file mode 100644
--- /dev/null
+++ b/Engine/Plugins/Developer/Concert/ConcertSync/ConcertSyncClient/Source/ConcertSyncClient/Private/Replication/Processing/ObjectInfo.h
@@ -0,0 +1,21 @@
+﻿// Copyright Epic Games, Inc. All Rights Reserved.
+
+#pragma once
+
+#include "Misc/Guid.h"
+#include "Replication/Data/ConcertPropertySelection.h"
+#include "Replication/Data/SequenceId.h"
+
+namespace UE::ConcertSyncClient
+{
+struct FObjectInfo
+{
+	/** The replication stream producing this object's data */
+	FGuid StreamId;
+	/** The properties to replicate */
+	FConcertPropertySelection SelectedProperties;
+			
+	/** Incremented every time replication data is sent out. Used for performance tracing. */
+	ConcertSyncCore::FSequenceId ReplicationSequenceId = 0;
+};
+}
diff --git a/Engine/Plugins/Developer/Concert/ConcertSync/ConcertSyncClient/Source/ConcertSyncClient/Private/Replication/Processing/ObjectReplicationApplierProcessor.cpp b/Engine/Plugins/Developer/Concert/ConcertSync/ConcertSyncClient/Source/ConcertSyncClient/Private/Replication/Processing/ObjectReplicationApplierProcessor.cpp
--- a/Engine/Plugins/Developer/Concert/ConcertSync/ConcertSyncClient/Source/ConcertSyncClient/Private/Replication/Processing/ObjectReplicationApplierProcessor.cpp
+++ b/Engine/Plugins/Developer/Concert/ConcertSync/ConcertSyncClient/Source/ConcertSyncClient/Private/Replication/Processing/ObjectReplicationApplierProcessor.cpp
@@ -8,6 +8,7 @@
 #include "Replication/Formats/IObjectReplicationFormat.h"
 #include "Replication/IConcertClientReplicationBridge.h"
 #include "Replication/Processing/Actions/ReplicationActionDispatcher.h"
+#include "Replication/Processing/IReplicationDataItem.h"
 #include "Trace/ConcertProtocolTrace.h"
 
 namespace UE::ConcertSyncClient::Replication
@@ -44,7 +45,7 @@ namespace UE::ConcertSyncClient::Replication
 		const bool bDebugActions = CVarDebugReplicationActions.GetValueOnGameThread();
 		
 		bool bAppliedData = false;
-		GetDataSource().ExtractReplicationDataForObject(Args.ObjectInfo.ObjectId,
+		Args.Item.ExtractReplicationDataForObject(
 			[this, &Args, Object, Actions, bDebugActions, &bAppliedData](const FConcertSessionSerializedPayload& Payload)
 			{
 				CONCERT_TRACE_REPLICATION_OBJECT_SCOPE(SerializeReceivedObject, Args.ObjectInfo.ObjectId.Object, Args.ObjectInfo.SequenceId);
diff --git a/Engine/Plugins/Developer/Concert/ConcertSync/ConcertSyncCore/Source/ConcertSyncCore/Private/Replication/Processing/ObjectReplicationProcessor.cpp b/Engine/Plugins/Developer/Concert/ConcertSync/ConcertSyncCore/Source/ConcertSyncCore/Private/Replication/Processing/ObjectReplicationProcessor.cpp
--- a/Engine/Plugins/Developer/Concert/ConcertSync/ConcertSyncCore/Source/ConcertSyncCore/Private/Replication/Processing/ObjectReplicationProcessor.cpp
+++ b/Engine/Plugins/Developer/Concert/ConcertSync/ConcertSyncCore/Source/ConcertSyncCore/Private/Replication/Processing/ObjectReplicationProcessor.cpp
@@ -16,9 +16,9 @@ namespace UE::ConcertSyncCore
 	void FObjectReplicationProcessor::ProcessObjects(const FProcessObjectsParams& Params)
 	{
 		// TODO UE-190714: Respect time budget and prioritize objects
-		DataSource.ForEachPendingObject([this](const FPendingObjectReplicationInfo& ObjectInfo)
+		DataSource.ConsumePendingObjects([this](const FPendingObjectReplicationInfo& ObjectInfo, IReplicationDataItem& Item)
 		{
-			ProcessObject({ ObjectInfo });
+			ProcessObject(FObjectProcessArgs(ObjectInfo, Item));
 		});
 	}
 }
diff --git a/Engine/Plugins/Developer/Concert/ConcertSync/ConcertSyncCore/Source/ConcertSyncCore/Private/Replication/Processing/ObjectReplicationSender.cpp b/Engine/Plugins/Developer/Concert/ConcertSync/ConcertSyncCore/Source/ConcertSyncCore/Private/Replication/Processing/ObjectReplicationSender.cpp
--- a/Engine/Plugins/Developer/Concert/ConcertSync/ConcertSyncCore/Source/ConcertSyncCore/Private/Replication/Processing/ObjectReplicationSender.cpp
+++ b/Engine/Plugins/Developer/Concert/ConcertSync/ConcertSyncCore/Source/ConcertSyncCore/Private/Replication/Processing/ObjectReplicationSender.cpp
@@ -4,6 +4,7 @@
 
 #include "ConcertLogGlobal.h"
 #include "IConcertSession.h"
+#include "Replication/Processing/IReplicationDataItem.h"
 #include "Replication/Processing/IReplicationDataSource.h"
 #include "Trace/ConcertProtocolTrace.h"
 
@@ -70,7 +71,7 @@ namespace UE::ConcertSyncCore
 				Forward<TPayloadRefType>(Payload)
 			});
 		};
-		GetDataSource().ExtractReplicationDataForObject(Args.ObjectInfo.ObjectId, CaptureData, CaptureData);
+		Args.Item.ExtractReplicationDataForObject(CaptureData, CaptureData);
 	}
 
 	inline void FObjectReplicationSender::MarkObjectForTrace(const FConcertReplicatedObjectId& Object, FSequenceId Id)
diff --git a/Engine/Plugins/Developer/Concert/ConcertSync/ConcertSyncCore/Source/ConcertSyncCore/Private/Replication/Processing/ReplicationDataQueuer.cpp b/Engine/Plugins/Developer/Concert/ConcertSync/ConcertSyncCore/Source/ConcertSyncCore/Private/Replication/Processing/ReplicationDataQueuer.cpp
--- a/Engine/Plugins/Developer/Concert/ConcertSync/ConcertSyncCore/Source/ConcertSyncCore/Private/Replication/Processing/ReplicationDataQueuer.cpp
+++ b/Engine/Plugins/Developer/Concert/ConcertSync/ConcertSyncCore/Source/ConcertSyncCore/Private/Replication/Processing/ReplicationDataQueuer.cpp
@@ -4,39 +4,56 @@
 
 #include "ConcertLogGlobal.h"
 #include "Replication/Messages/ObjectReplication.h"
+#include "Replication/Processing/IReplicationDataItem.h"
 #include "Trace/ConcertProtocolTrace.h"
 
 namespace UE::ConcertSyncCore
 {
-	void FReplicationDataQueuer::ForEachPendingObject(TFunctionRef<void(const FPendingObjectReplicationInfo&)> ProcessItemFunc) const
+namespace DataQueuer
+{
+	class FDataQueuerItem : public IReplicationDataItem
 	{
-		for (auto It = PendingEvents.CreateConstIterator(); It; ++It)
+		TSharedPtr<const FConcertReplication_ObjectReplicationEvent> DataToApply;
+		TFunctionRef<void()> RemoveItemFunc;
+		bool bAlreadyExtracted = false;
+	public:
+		
+		explicit FDataQueuerItem(
+			const TSharedPtr<const FConcertReplication_ObjectReplicationEvent>& DataToApply, 
+			const TFunctionRef<void()>& RemoveItemFunc
+			)
+			: DataToApply(DataToApply)
+			, RemoveItemFunc(RemoveItemFunc)
+		{}
+
+		virtual void ExtractReplicationDataForObject(
+			TFunctionRef<void(const FConcertSessionSerializedPayload& Payload)> ProcessCopyable,
+			TFunctionRef<void(FConcertSessionSerializedPayload&& Payload)> ProcessMoveable
+			) override
 		{
-			ProcessItemFunc({ It->Key, It->Value.SequenceId });
+			if (!ensure(!bAlreadyExtracted))
+			{
+				return;
+			}
+			bAlreadyExtracted = true;
+			
+			RemoveItemFunc();
+			// Event may be shared by other FReplicationDataQueuers since it originates from the replication cache, so sadly no move.
+			ProcessCopyable(DataToApply->SerializedPayload);
 		}
-	}
-
-	int32 FReplicationDataQueuer::NumObjects() const
-	{
-		return PendingEvents.Num();
-	}
-
-	bool FReplicationDataQueuer::ExtractReplicationDataForObject(
-		const FConcertReplicatedObjectId& Object,
-		TFunctionRef<void(const FConcertSessionSerializedPayload& Payload)> ProcessCopyable,
-		TFunctionRef<void(FConcertSessionSerializedPayload&& Payload)> ProcessMoveable
+	};
+}
+	void FReplicationDataQueuer::ConsumePendingObjects(
+		TFunctionRef<void(const FPendingObjectReplicationInfo& Info, IReplicationDataItem& Item)> InProcess
 		)
 	{
-		FPendingObjectData EventData;
-		const bool bSuccess = PendingEvents.RemoveAndCopyValue(Object, EventData);
-		if (!ensureMsgf(bSuccess, TEXT("ExtractReplicationDataForObject for an item that was not returned by ForEachPendingObject")))
+		for (auto It = PendingEvents.CreateIterator(); It; ++It)
 		{
-			return false;
+			const FPendingObjectReplicationInfo Info{ It->Key, It->Value.SequenceId };
+			const auto RemoveItem = [&It]{ It.RemoveCurrent(); };
+			DataQueuer::FDataQueuerItem Item(It->Value.DataToApply, RemoveItem);
+			InProcess(Info, Item);
 		}
-
-		// Event may be shared by other FReplicationDataQueuers since it originates from the replication cache, so sadly no move.
-		ProcessCopyable(EventData.DataToApply->SerializedPayload);
-		return true;
 	}
 
 	void FReplicationDataQueuer::OnDataCached(const FConcertReplicatedObjectId& Object, FSequenceId SequenceId, TSharedRef<const FConcertReplication_ObjectReplicationEvent> Data)
diff --git a/Engine/Plugins/Developer/Concert/ConcertSync/ConcertSyncCore/Source/ConcertSyncCore/Public/Replication/Processing/IReplicationDataItem.h b/Engine/Plugins/Developer/Concert/ConcertSync/ConcertSyncCore/Source/ConcertSyncCore/Public/Replication/Processing/IReplicationDataItem.h
new file mode 100644
--- /dev/null
+++ b/Engine/Plugins/Developer/Concert/ConcertSync/ConcertSyncCore/Source/ConcertSyncCore/Public/Replication/Processing/IReplicationDataItem.h
@@ -0,0 +1,38 @@
+﻿// Copyright Epic Games, Inc. All Rights Reserved.
+
+#pragma once
+
+#include "Templates/Function.h"
+
+struct FConcertSessionSerializedPayload;
+
+namespace UE::ConcertSyncCore
+{
+/** An item that is cached by IReplicationDataSource. Provides interface for extracting the replication data.*/
+class IReplicationDataItem
+{
+public:
+	
+	/**
+	 * Extracts data for an object iterated by IReplicationDataSource.
+	 * 
+	 * This can update the SequenceId associated with the object (if this data source generates data as opposed to queuing it).
+	 * Either ProcessCopyable or ProcessMoveable will be called, never both, and it will be called at most once.
+	 * 
+	 * @param ProcessCopyable Callback if the event was retrieved and not owned by this IReplicationDataSource (hence not being moveable).
+	 * @param ProcessMoveable Callback if the event was just constructed (and hence can be moved)	
+	 */
+	virtual void ExtractReplicationDataForObject(
+		TFunctionRef<void(const FConcertSessionSerializedPayload& Payload)> ProcessCopyable,
+		TFunctionRef<void(FConcertSessionSerializedPayload&& Payload)> ProcessMoveable
+		) = 0;
+	
+	/** Util version for callers that only want to read and do not want to store the payload. */
+	void ExtractReplicationDataForObject(TFunctionRef<void(const FConcertSessionSerializedPayload& Payload)> ProcessCopyable)
+	{
+		ExtractReplicationDataForObject(ProcessCopyable, [&ProcessCopyable](FConcertSessionSerializedPayload&& Payload){ ProcessCopyable(Payload); });
+	}
+	
+	virtual ~IReplicationDataItem() = default;
+};
+}
diff --git a/Engine/Plugins/Developer/Concert/ConcertSync/ConcertSyncCore/Source/ConcertSyncCore/Public/Replication/Processing/IReplicationDataSource.h b/Engine/Plugins/Developer/Concert/ConcertSync/ConcertSyncCore/Source/ConcertSyncCore/Public/Replication/Processing/IReplicationDataSource.h
--- a/Engine/Plugins/Developer/Concert/ConcertSync/ConcertSyncCore/Source/ConcertSyncCore/Public/Replication/Processing/IReplicationDataSource.h
+++ b/Engine/Plugins/Developer/Concert/ConcertSync/ConcertSyncCore/Source/ConcertSyncCore/Public/Replication/Processing/IReplicationDataSource.h
@@ -12,6 +12,8 @@
 
 namespace UE::ConcertSyncCore
 {
+	class IReplicationDataItem;
+
 	struct FPendingObjectReplicationInfo
 	{
 		/** The object to be replicated */
@@ -34,11 +36,21 @@ namespace UE::ConcertSyncCore
 	{
 	public:
 		
+		/**
+		 * Iterates the objects that must be processed for replication.
+		 * 
+		 * It allows the extraction of the data using IReplicationDataItem. If the data is extracted, the data is "dequeued", i.e. will not be 
+		 * passed again in subsequent calls to ConsumePendingObjects unless new data becomes available.
+		 */
+		virtual void ConsumePendingObjects(TFunctionRef<void(const FPendingObjectReplicationInfo& Info, IReplicationDataItem& Item)> InProcess) = 0;
+		
 		/** Iterates the objects that must be processed for replication (the result of ExtractReplicationDataForObject MAY return something new). */
-		virtual void ForEachPendingObject(TFunctionRef<void(const FPendingObjectReplicationInfo&)> ProcessItemFunc) const = 0;
+		UE_DEPRECATED(5.8, "Use ConsumePendingObjects instead.")
+		virtual void ForEachPendingObject(TFunctionRef<void(const FPendingObjectReplicationInfo&)> ProcessItemFunc) const {}
+		
 		/** @return The number of objects ForEachPendingObject would iterate. Can be used, e.g. for Reserve()-ing a container. */
-		virtual int32 NumObjects() const = 0;
-
+		UE_DEPRECATED(5.8, "This function no longer has a use and will be removed.")
+		virtual int32 NumObjects() const { return 0; }
 		/**
 		 * Extracts data for Object. Object must have been previously returned by ForEachPendingObject.
 		 * 
@@ -53,15 +65,21 @@ namespace UE::ConcertSyncCore
 		 * 
 		 * @return Whether successful. False indicates the call was invalid to make (ForEachPendingObject did not return Object). True indicates success, even if ProcessCopyable was not called.		
 		 */
+		UE_DEPRECATED(5.8, "Use ConsumePendingObjects and call IReplicationDataItem::ExtractReplicationDataForObject instead.")
 		virtual bool ExtractReplicationDataForObject(
 			const FConcertReplicatedObjectId& Object,
 			TFunctionRef<void(const FConcertSessionSerializedPayload& Payload)> ProcessCopyable,
 			TFunctionRef<void(FConcertSessionSerializedPayload&& Payload)> ProcessMoveable
-			) = 0;
+			)
+		{
+			return false;
+		}
+		
 		/** Util version for callers that only want to read and do not want to store the payload. */
+		UE_DEPRECATED(5.8, "Use ConsumePendingObjects and call IReplicationDataItem::ExtractReplicationDataForObject instead.")
 		bool ExtractReplicationDataForObject(const FConcertReplicatedObjectId& Object, TFunctionRef<void(const FConcertSessionSerializedPayload& Payload)> ProcessCopyable)
 		{
-			return ExtractReplicationDataForObject(Object, ProcessCopyable, [&ProcessCopyable](FConcertSessionSerializedPayload&& Payload){ ProcessCopyable(Payload); });
+			return false;
 		}
 		
 		virtual ~IReplicationDataSource() = default;
diff --git a/Engine/Plugins/Developer/Concert/ConcertSync/ConcertSyncCore/Source/ConcertSyncCore/Public/Replication/Processing/ObjectReplicationProcessor.h b/Engine/Plugins/Developer/Concert/ConcertSync/ConcertSyncCore/Source/ConcertSyncCore/Public/Replication/Processing/ObjectReplicationProcessor.h
--- a/Engine/Plugins/Developer/Concert/ConcertSync/ConcertSyncCore/Source/ConcertSyncCore/Public/Replication/Processing/ObjectReplicationProcessor.h
+++ b/Engine/Plugins/Developer/Concert/ConcertSync/ConcertSyncCore/Source/ConcertSyncCore/Public/Replication/Processing/ObjectReplicationProcessor.h
@@ -56,6 +56,13 @@ namespace UE::ConcertSyncCore
 		{
 			/** Info about the object to process */
 			FPendingObjectReplicationInfo ObjectInfo;
+			
+			/** Item you can use to extract the replication info from. */
+			IReplicationDataItem& Item;
+
+			explicit FObjectProcessArgs(FPendingObjectReplicationInfo InObjectInfo, IReplicationDataItem& InItem)
+				: ObjectInfo(MoveTemp(InObjectInfo)) , Item(InItem)
+			{}
 		};
 
 		FORCEINLINE IReplicationDataSource& GetDataSource() const { return DataSource; }
diff --git a/Engine/Plugins/Developer/Concert/ConcertSync/ConcertSyncCore/Source/ConcertSyncCore/Public/Replication/Processing/ReplicationDataQueuer.h b/Engine/Plugins/Developer/Concert/ConcertSync/ConcertSyncCore/Source/ConcertSyncCore/Public/Replication/Processing/ReplicationDataQueuer.h
--- a/Engine/Plugins/Developer/Concert/ConcertSync/ConcertSyncCore/Source/ConcertSyncCore/Public/Replication/Processing/ReplicationDataQueuer.h
+++ b/Engine/Plugins/Developer/Concert/ConcertSync/ConcertSyncCore/Source/ConcertSyncCore/Public/Replication/Processing/ReplicationDataQueuer.h
@@ -27,9 +27,7 @@ namespace UE::ConcertSyncCore
 	public:
 
 		//~ Begin IReplicationDataSource Interface
-		UE_API virtual void ForEachPendingObject(TFunctionRef<void(const FPendingObjectReplicationInfo&)> ProcessItemFunc) const override;
-		UE_API virtual int32 NumObjects() const override;
-		UE_API virtual bool ExtractReplicationDataForObject(const FConcertReplicatedObjectId& Object, TFunctionRef<void(const FConcertSessionSerializedPayload& Payload)> ProcessCopyable, TFunctionRef<void(FConcertSessionSerializedPayload&& Payload)> ProcessMoveable) override;
+		UE_API virtual void ConsumePendingObjects(TFunctionRef<void(const FPendingObjectReplicationInfo& Info, IReplicationDataItem& Item)> InProcess) override;
 		//~ End IReplicationDataSource Interface
 		
 		//~ Begin IReplicationCacheUser Interface
